Adding weights to bundle manager status.
Adding helper class to monitor bundle + prereq status as a single %.
#tests iOS Client

diff --git a/Engine/Source/Runtime/Core/Public/GenericPlatform/GenericPlatformInstallBundleManager.h b/Engine/Source/Runtime/Core/Public/GenericPlatform/GenericPlatformInstallBundleManager.h
--- a/Engine/Source/Runtime/Core/Public/GenericPlatform/GenericPlatformInstallBundleManager.h
+++ b/Engine/Source/Runtime/Core/Public/GenericPlatform/GenericPlatformInstallBundleManager.h
@@ -230,8 +230,9 @@ struct FInstallBundleContentState
 {
 	EInstallBundleContentState State = EInstallBundleContentState::InitializationError;
 	TMap<FName, EInstallBundleContentState> IndividualBundleStates;
+	TMap<FName, float> IndividualBundleWeights;
 	uint64 DownloadSize = 0;
-	uint64 InstallSize = 0;	
+	uint64 InstallSize = 0;
 	uint64 FreeSpace = 0;
 };
 
diff --git a/Engine/Source/Runtime/InstallBundleManager/InstallBundleManager.Build.cs b/Engine/Source/Runtime/InstallBundleManager/InstallBundleManager.Build.cs
new file mode 100644
--- /dev/null
+++ b/Engine/Source/Runtime/InstallBundleManager/InstallBundleManager.Build.cs
@@ -0,0 +1,22 @@
+// Copyright 1998-2019 Epic Games, Inc. All Rights Reserved.
+
+using UnrealBuildTool;
+
+public class InstallBundleManager : ModuleRules
+{
+	public InstallBundleManager(ReadOnlyTargetRules Target) : base(Target)
+	{
+		PrivatePCHHeaderFile = "Private/InstallBundleManagerPrivatePCH.h";
+
+		PrivateDependencyModuleNames.AddRange(
+			new string[] {
+				"Core",
+				"CoreUObject",
+				"Engine",
+				"InputCore",
+				"Projects",
+				"ApplicationCore",
+			}
+		);
+	}
+}
diff --git a/Engine/Source/Runtime/InstallBundleManager/Private/BundlePrereqCombinedStatusHelper.cpp b/Engine/Source/Runtime/InstallBundleManager/Private/BundlePrereqCombinedStatusHelper.cpp
new file mode 100644
--- /dev/null
+++ b/Engine/Source/Runtime/InstallBundleManager/Private/BundlePrereqCombinedStatusHelper.cpp
@@ -0,0 +1,255 @@
+// Copyright 1998-2019 Epic Games, Inc. All Rights Reserved.
+
+#include "BundlePrereqCombinedStatusHelper.h"
+#include "Containers/Ticker.h"
+#include "InstallBundleManagerPrivatePCH.h"
+
+FBundlePrereqCombinedStatusHelper::FBundlePrereqCombinedStatusHelper()
+: RequiredBundleNames()
+, BundleStatusCache()
+, CachedBundleWeights()
+, CurrentCombinedStatus()
+, DownloadWeight(1.f)
+, InstallWeight(1.f)
+, bBundleNeedsUpdate(false)
+, InstallBundleManager(nullptr)
+, TickHandle()
+{
+	SetupDelegates();
+}
+
+FBundlePrereqCombinedStatusHelper::~FBundlePrereqCombinedStatusHelper()
+{
+	CleanUpDelegates();
+}
+
+void FBundlePrereqCombinedStatusHelper::SetupDelegates()
+{
+	IPlatformInstallBundleManager::InstallBundleCompleteDelegate.AddRaw(this, &FBundlePrereqCombinedStatusHelper::OnBundleInstallComplete);
+	if (!TickHandle.IsValid())
+	{
+		TickHandle = FTicker::GetCoreTicker().AddTicker(FTickerDelegate::CreateRaw(this, &FBundlePrereqCombinedStatusHelper::Tick));
+	}
+}
+
+void FBundlePrereqCombinedStatusHelper::CleanUpDelegates()
+{
+	IPlatformInstallBundleManager::InstallBundleCompleteDelegate.RemoveAll(this);
+	if (TickHandle.IsValid())
+	{
+		FTicker::GetCoreTicker().RemoveTicker(TickHandle);
+		TickHandle.Reset();
+	}
+}
+
+void FBundlePrereqCombinedStatusHelper::SetBundlesToTrackFromContentState(FInstallBundleContentState& BundleContentState)
+{
+	RequiredBundleNames.Empty();
+	CachedBundleWeights.Empty();
+	
+	//Track if we need any kind of bundle updates
+	if (BundleContentState.State == EInstallBundleContentState::NotInstalled || BundleContentState.State == EInstallBundleContentState::NeedsUpdate)
+	{
+		bBundleNeedsUpdate = true;
+	}
+	
+	//Save required bundles and their weights
+	for (TPair<FName, float>& BundleStatePair : BundleContentState.IndividualBundleWeights)
+	{
+		RequiredBundleNames.Add(BundleStatePair.Key);
+		
+		float& CachedBundleWeight = CachedBundleWeights.FindOrAdd(BundleStatePair.Key);
+		CachedBundleWeight = BundleStatePair.Value;
+	}
+	
+	//Go ahead and calculate initial values from the Bundle Cache
+	UpdateBundleCache();
+}
+
+void FBundlePrereqCombinedStatusHelper::UpdateBundleCache()
+{
+	//if we haven't already set this up, lets try to set it now
+	if (nullptr == InstallBundleManager)
+	{
+		InstallBundleManager = FPlatformMisc::GetPlatformInstallBundleManager();
+	}
+	
+	if (ensureAlwaysMsgf((nullptr != InstallBundleManager), TEXT("Invalid InstallBundleManager during UpdateBundleCache! Needs to be valid during run!")))
+	{
+		for (FName& BundleName : RequiredBundleNames)
+		{
+			TOptional<FInstallBundleStatus> BundleProgress = InstallBundleManager->GetBundleProgress(BundleName);
+			
+			//Copy progress to the cache as long as we have progress to copy.
+			if (BundleProgress.IsSet())
+			{
+				FInstallBundleStatus& CachedStatus = BundleStatusCache.FindOrAdd(BundleName);
+				CachedStatus = BundleProgress.GetValue();
+			}
+		}
+	}
+}
+
+void FBundlePrereqCombinedStatusHelper::UpdateCombinedStatus()
+{
+	CurrentCombinedStatus.ProgressPercent = GetCombinedProgressPercent();
+	
+	EInstallBundleStatus EarliestBundleState = EInstallBundleStatus::Count;
+	bool bIsAnythingPaused = false;
+	
+	for (const TPair<FName,FInstallBundleStatus>& BundlePair : BundleStatusCache)
+	{
+		if (BundlePair.Value.Status < EarliestBundleState)
+		{
+			EarliestBundleState = BundlePair.Value.Status;
+		}
+		
+		bIsAnythingPaused = (bIsAnythingPaused || (BundlePair.Value.PauseFlags != EInstallBundlePauseFlags::None));
+	}
+	
+	//if we have any paused bundles, and we have any bundle that isn't finished installed, we are Paused
+	//if everything is installed ignore the pause flags as we completed after pausing the bundles
+	CurrentCombinedStatus.bIsPaused = (bIsAnythingPaused && (EarliestBundleState < EInstallBundleStatus::Installed));
+	
+	//Set to true by default, some cases below will turn this off as we don't support pausing during Mounting or Compiling Shaders
+	CurrentCombinedStatus.bDoesCurrentStateSupportPausing = true;
+	
+	if (CurrentCombinedStatus.ProgressPercent == 0.f)
+	{
+		CurrentCombinedStatus.CombinedState = FCombinedBundleStatus::ECombinedBundleStateEnum::Initializing;
+	}
+	else if (!bBundleNeedsUpdate)
+	{
+		CurrentCombinedStatus.CombinedState = FCombinedBundleStatus::ECombinedBundleStateEnum::NoUpdateRequired;
+		CurrentCombinedStatus.bDoesCurrentStateSupportPausing = false;
+	}
+	else if (EarliestBundleState < EInstallBundleStatus::Finishing)
+	{
+		CurrentCombinedStatus.CombinedState = FCombinedBundleStatus::ECombinedBundleStateEnum::Updating;
+	}
+	else if (EarliestBundleState < EInstallBundleStatus::Installed)
+	{
+		CurrentCombinedStatus.CombinedState = FCombinedBundleStatus::ECombinedBundleStateEnum::CompilingShaders;
+		CurrentCombinedStatus.bDoesCurrentStateSupportPausing = false;
+	}
+	else if (EarliestBundleState == EInstallBundleStatus::Installed)
+	{
+		CurrentCombinedStatus.CombinedState = FCombinedBundleStatus::ECombinedBundleStateEnum::Finished;
+		CurrentCombinedStatus.bDoesCurrentStateSupportPausing = false;
+	}
+	else
+	{
+		CurrentCombinedStatus.CombinedState = FCombinedBundleStatus::ECombinedBundleStateEnum::Unknown;
+	}
+}
+
+float FBundlePrereqCombinedStatusHelper::GetCombinedProgressPercent()
+{
+	float AllBundleProgressPercent = 0.f;
+	
+	ensureAlwaysMsgf((CachedBundleWeights.Num() >= BundleStatusCache.Num()), TEXT("Missing Cache Entries for BundleWeights!Any missing bundles will have 0 for their progress!"));
+	
+	for (TPair<FName,FInstallBundleStatus>& BundleStatusPair : BundleStatusCache)
+	{
+		float* FoundWeight = CachedBundleWeights.Find(BundleStatusPair.Key);
+		if (ensureAlwaysMsgf((nullptr != FoundWeight), TEXT("Found missing entry for BundleWeight! Bundle %s does not have a weight entry!"), *(BundleStatusPair.Key.ToString())))
+		{
+			AllBundleProgressPercent += ((*FoundWeight) * GetIndividualWeightedProgressPercent(BundleStatusPair.Value));
+		}
+	}
+	
+	FMath::Clamp(AllBundleProgressPercent, 0.f, 1.0f);
+	return AllBundleProgressPercent;
+}
+
+float FBundlePrereqCombinedStatusHelper::GetIndividualWeightedProgressPercent(FInstallBundleStatus& BundleStatus)
+{
+	const float TotalWeight = DownloadWeight + InstallWeight;
+	float CombinedOverallProgressPercent = 0.f;
+	
+	//Completed is maximum progress possible
+	if (BundleStatus.Status >= EInstallBundleStatus::Installed)
+	{
+		CombinedOverallProgressPercent = 1.0f;
+	}
+	//Once we are in finishing, we display a new bar for this step, so just show raw Finishing_Percent
+	else if (BundleStatus.Status >= EInstallBundleStatus::Finishing)
+	{
+		CombinedOverallProgressPercent = BundleStatus.Finishing_Percent;
+	}
+	else
+	{
+		if (ensureAlwaysMsgf((TotalWeight > 0), TEXT("Invalid Weights for FBundleStatusTracker! Need to have weights > 0.f!")))
+		{
+			float DownloadWeightedPercent = 0.f;
+			bool bDidHaveBackgroundDownloadProgress = false;
+			
+			if (DownloadWeight > 0.f)
+			{
+				//Skip background downloads from contributing to progress if they aren't in use
+				bDidHaveBackgroundDownloadProgress = (BundleStatus.BackgroundDownloadProgress.IsSet() && (BundleStatus.BackgroundDownloadProgress->BytesDownloaded > 0));
+				if (bDidHaveBackgroundDownloadProgress)
+				{
+					const float AppliedDownloadWeight = DownloadWeight / TotalWeight;
+					
+					if (BundleStatus.Status > EInstallBundleStatus::Downloading)
+					{
+						//if we have moved past the download phase, just consider these
+						//downloads as having finished
+						DownloadWeightedPercent = 1.0f * AppliedDownloadWeight;
+					}
+					else
+					{
+						DownloadWeightedPercent = BundleStatus.BackgroundDownloadProgress->PercentComplete * AppliedDownloadWeight;
+					}
+				}
+			}
+			
+			float InstallWeightedPercent = 0.f;
+			if (BundleStatus.InstallDownloadProgress.IsSet())
+			{
+				//If we didn't have background download progress, just use our InstallWeight as we don't have Download Progress to use
+				const float AppliedInstallWeight = bDidHaveBackgroundDownloadProgress ? (InstallWeight / TotalWeight) : InstallWeight;
+				InstallWeightedPercent = BundleStatus.Install_Percent * AppliedInstallWeight;
+			}
+			
+			CombinedOverallProgressPercent = DownloadWeightedPercent + InstallWeightedPercent;
+			FMath::Clamp(CombinedOverallProgressPercent, 0.f, 1.0f);
+		}
+	}
+	
+	return CombinedOverallProgressPercent;
+}
+
+bool FBundlePrereqCombinedStatusHelper::Tick(float dt)
+{
+	UpdateBundleCache();
+	UpdateCombinedStatus();
+	
+	//just always keep ticking
+	return true;
+}
+
+const FBundlePrereqCombinedStatusHelper::FCombinedBundleStatus& FBundlePrereqCombinedStatusHelper::GetCurrentCombinedState() const
+{
+	return CurrentCombinedStatus;
+}
+
+void FBundlePrereqCombinedStatusHelper::OnBundleInstallComplete(FInstallBundleResultInfo CompletedBundleInfo)
+{
+	const FName CompletedBundleName = CompletedBundleInfo.BundleName;
+	const bool bBundleCompletedSuccessfully = (CompletedBundleInfo.Result == EInstallBundleResult::OK);
+	
+	if (bBundleCompletedSuccessfully)
+	{
+		//Make sure our BundleCache shows this as finished all the way
+		FInstallBundleStatus& BundleCacheInfo = BundleStatusCache.FindOrAdd(CompletedBundleName);
+		BundleCacheInfo.Status = EInstallBundleStatus::Installed;
+		
+		TOptional<FInstallBundleStatus> BundleProgress = InstallBundleManager->GetBundleProgress(CompletedBundleName);
+		if (ensureAlwaysMsgf(BundleProgress.IsSet(), TEXT("Expected to find BundleProgress for completed bundle, but did not. Leaving old progress values")))
+		{
+			BundleCacheInfo = BundleProgress.GetValue();
+		}
+	}
+}
diff --git a/Engine/Source/Runtime/InstallBundleManager/Private/InstallBundleManagerPrivatePCH.h b/Engine/Source/Runtime/InstallBundleManager/Private/InstallBundleManagerPrivatePCH.h
new file mode 100644
--- /dev/null
+++ b/Engine/Source/Runtime/InstallBundleManager/Private/InstallBundleManagerPrivatePCH.h
@@ -0,0 +1,6 @@
+// Copyright 1998-2019 Epic Games, Inc. All Rights Reserved.
+
+#pragma once
+
+#include "CoreMinimal.h"
+#include "GenericPlatform/GenericPlatformInstallBundleManager.h"
diff --git a/Engine/Source/Runtime/InstallBundleManager/Public/BundlePrereqCombinedStatusHelper.h b/Engine/Source/Runtime/InstallBundleManager/Public/BundlePrereqCombinedStatusHelper.h
new file mode 100644
--- /dev/null
+++ b/Engine/Source/Runtime/InstallBundleManager/Public/BundlePrereqCombinedStatusHelper.h
@@ -0,0 +1,85 @@
+// Copyright 1998-2019 Epic Games, Inc. All Rights Reserved.
+#pragma once
+
+#include "CoreMinimal.h"
+#include "GenericPlatform/GenericPlatformInstallBundleManager.h"
+
+//Handles calculating the bundle status by combining progress from all of its
+//Prerequisites. Allows you to display one progress percent that is weighted based on all
+//bundles' values.
+class FBundlePrereqCombinedStatusHelper
+{
+public:
+	//provide all our needed combined status information in 1 struct
+	struct FCombinedBundleStatus
+	{
+		//Collapses all the bundle manager states into one of a few states so that you can show simple text based on this enum
+		enum class ECombinedBundleStateEnum : int32
+		{
+			Unknown
+			,Initializing
+			,Updating
+			,CompilingShaders
+			,NoUpdateRequired
+			,Finished
+		};
+		
+		float ProgressPercent;
+		ECombinedBundleStateEnum CombinedState;
+		bool bIsPaused;
+		bool bDoesCurrentStateSupportPausing;
+		
+		FCombinedBundleStatus()
+		: ProgressPercent(0.f)
+		, CombinedState(ECombinedBundleStateEnum::Unknown)
+		, bIsPaused(false)
+		, bDoesCurrentStateSupportPausing(false)
+		{
+		}
+	};
+	
+public:
+	FBundlePrereqCombinedStatusHelper();
+	~FBundlePrereqCombinedStatusHelper();
+	
+	//Setup tracking for all bundles required in the supplied BundleContentState
+	void SetBundlesToTrackFromContentState(FInstallBundleContentState& BundleContentState);
+	
+	//Get current CombinedBundleStatus for everything setup to track
+	const FCombinedBundleStatus& GetCurrentCombinedState() const;
+	
+private:
+	bool Tick(float dt);
+	void UpdateBundleCache();
+	void UpdateCombinedStatus();
+	
+	void SetupDelegates();
+	void CleanUpDelegates();
+	
+	//Called so we can track when a bundle is finished
+	void OnBundleInstallComplete(FInstallBundleResultInfo CompletedBundleInfo);
+	
+	float GetCombinedProgressPercent();
+	float GetIndividualWeightedProgressPercent(FInstallBundleStatus& Bundle);
+	
+private:
+	//All bundles we need including pre-reqs
+	TArray<FName> RequiredBundleNames;
+	
+	//Internal Cache of all bundle statuses to track progress
+	TMap<FName, FInstallBundleStatus> BundleStatusCache;
+	
+	//Bundle weights that determine what % of the overall install each bundle represents
+	TMap<FName, float> CachedBundleWeights;
+	
+	FCombinedBundleStatus CurrentCombinedStatus;
+	
+	//How to weight downloads vs. installs
+	float DownloadWeight;
+	float InstallWeight;
+	
+	bool bBundleNeedsUpdate;
+	
+	IPlatformInstallBundleManager* InstallBundleManager;
+	FDelegateHandle TickHandle;
+};
