Adding InstallBundleManager placeholder module files. Fixing broken editor build.
#tests IOS Client

diff --git a/Engine/Source/Runtime/InstallBundleManager/Private/InstallBundleModule.cpp b/Engine/Source/Runtime/InstallBundleManager/Private/InstallBundleModule.cpp
new file mode 100644
--- /dev/null
+++ b/Engine/Source/Runtime/InstallBundleManager/Private/InstallBundleModule.cpp
@@ -0,0 +1,4 @@
+// Copyright 1998-2019 Epic Games, Inc. All Rights Reserved.
+#include "InstallBundleModule.h"
+
+IMPLEMENT_MODULE(FInstallBundleModule, InstallBundle);
diff --git a/Engine/Source/Runtime/InstallBundleManager/Public/InstallBundleModule.h b/Engine/Source/Runtime/InstallBundleManager/Public/InstallBundleModule.h
new file mode 100644
--- /dev/null
+++ b/Engine/Source/Runtime/InstallBundleManager/Public/InstallBundleModule.h
@@ -0,0 +1,21 @@
+// Copyright 1998-2019 Epic Games, Inc. All Rights Reserved.
+#pragma once
+
+#include "Modules/ModuleInterface.h"
+#include "Modules/ModuleManager.h"
+
+/**
+ * Currently empty implementation for InstallBundleModule until things are moved in here.
+ */
+class FInstallBundleModule : public IModuleInterface
+{
+public:
+	// IModuleInterface interface
+	virtual void StartupModule() override
+	{
+	}
+
+	virtual void ShutdownModule() override
+	{
+	}
+};
