Renaming InstallBundleManager module to be consistent (was InstallBundle instead of InstallBundleManager). Should hopefully now actually fix the editor build.
#tests none

diff --git a/Engine/Source/Runtime/InstallBundleManager/Private/InstallBundleManagerModule.cpp b/Engine/Source/Runtime/InstallBundleManager/Private/InstallBundleManagerModule.cpp
new file mode 100644
--- /dev/null
+++ b/Engine/Source/Runtime/InstallBundleManager/Private/InstallBundleManagerModule.cpp
@@ -0,0 +1,4 @@
+// Copyright 1998-2019 Epic Games, Inc. All Rights Reserved.
+#include "InstallBundleManagerModule.h"
+
+IMPLEMENT_MODULE(FInstallBundleManagerModule, InstallBundleManager);
diff --git a/Engine/Source/Runtime/InstallBundleManager/Private/InstallBundleModule.cpp b/Engine/Source/Runtime/InstallBundleManager/Private/InstallBundleModule.cpp
deleted file mode 100644
diff --git a/Engine/Source/Runtime/InstallBundleManager/Public/InstallBundleManagerModule.h b/Engine/Source/Runtime/InstallBundleManager/Public/InstallBundleManagerModule.h
new file mode 100644
--- /dev/null
+++ b/Engine/Source/Runtime/InstallBundleManager/Public/InstallBundleManagerModule.h
@@ -0,0 +1,21 @@
+// Copyright 1998-2019 Epic Games, Inc. All Rights Reserved.
+#pragma once
+
+#include "Modules/ModuleInterface.h"
+#include "Modules/ModuleManager.h"
+
+/**
+ * Currently empty implementation for InstallBundleModule until things are moved in here.
+ */
+class FInstallBundleManagerModule : public IModuleInterface
+{
+public:
+	// IModuleInterface interface
+	virtual void StartupModule() override
+	{
+	}
+
+	virtual void ShutdownModule() override
+	{
+	}
+};
diff --git a/Engine/Source/Runtime/InstallBundleManager/Public/InstallBundleModule.h b/Engine/Source/Runtime/InstallBundleManager/Public/InstallBundleModule.h
deleted file mode 100644
