Exporting new BundlePrereqCombinedStatusHelper methods in module.
#tests none

diff --git a/Engine/Source/Runtime/InstallBundleManager/Public/BundlePrereqCombinedStatusHelper.h b/Engine/Source/Runtime/InstallBundleManager/Public/BundlePrereqCombinedStatusHelper.h
--- a/Engine/Source/Runtime/InstallBundleManager/Public/BundlePrereqCombinedStatusHelper.h
+++ b/Engine/Source/Runtime/InstallBundleManager/Public/BundlePrereqCombinedStatusHelper.h
@@ -7,7 +7,7 @@
 //Handles calculating the bundle status by combining progress from all of its
 //Prerequisites. Allows you to display one progress percent that is weighted based on all
 //bundles' values.
-class FBundlePrereqCombinedStatusHelper
+class INSTALLBUNDLEMANAGER_API FBundlePrereqCombinedStatusHelper
 {
 public:
 	//provide all our needed combined status information in 1 struct
