Fixing ensure on multiple calls to BundlePrereqCombinedStatusHelper::SetBundlesToTrackFromContentState
Fixing ensure when bundles you aren't tracking complete while running a BundlePrereqCombinedStatusHelper
#tests iOS Client

diff --git a/Engine/Source/Runtime/InstallBundleManager/Private/BundlePrereqCombinedStatusHelper.cpp b/Engine/Source/Runtime/InstallBundleManager/Private/BundlePrereqCombinedStatusHelper.cpp
--- a/Engine/Source/Runtime/InstallBundleManager/Private/BundlePrereqCombinedStatusHelper.cpp
+++ b/Engine/Source/Runtime/InstallBundleManager/Private/BundlePrereqCombinedStatusHelper.cpp
@@ -46,6 +46,7 @@ void FBundlePrereqCombinedStatusHelper::SetBundlesToTrackFromContentState(FInsta
 {
 	RequiredBundleNames.Empty();
 	CachedBundleWeights.Empty();
+	BundleStatusCache.Empty();
 	
 	//Track if we need any kind of bundle updates
 	if (BundleContentState.State == EInstallBundleContentState::NotInstalled || BundleContentState.State == EInstallBundleContentState::NeedsUpdate)
@@ -57,9 +58,7 @@ void FBundlePrereqCombinedStatusHelper::SetBundlesToTrackFromContentState(FInsta
 	for (TPair<FName, float>& BundleStatePair : BundleContentState.IndividualBundleWeights)
 	{
 		RequiredBundleNames.Add(BundleStatePair.Key);
-		
-		float& CachedBundleWeight = CachedBundleWeights.FindOrAdd(BundleStatePair.Key);
-		CachedBundleWeight = BundleStatePair.Value;
+		CachedBundleWeights.FindOrAdd(BundleStatePair.Key) = BundleStatePair.Value;
 	}
 	
 	//Go ahead and calculate initial values from the Bundle Cache
@@ -239,8 +238,9 @@ void FBundlePrereqCombinedStatusHelper::OnBundleInstallComplete(FInstallBundleRe
 {
 	const FName CompletedBundleName = CompletedBundleInfo.BundleName;
 	const bool bBundleCompletedSuccessfully = (CompletedBundleInfo.Result == EInstallBundleResult::OK);
+	const bool bWasRequiredBundle = RequiredBundleNames.Contains(CompletedBundleInfo.BundleName);
 	
-	if (bBundleCompletedSuccessfully)
+	if (bBundleCompletedSuccessfully && bWasRequiredBundle)
 	{
 		//Make sure our BundleCache shows this as finished all the way
 		FInstallBundleStatus& BundleCacheInfo = BundleStatusCache.FindOrAdd(CompletedBundleName);
