Adding Logging into InstallBundleManager module.
Adding Move Constructor implementation first pass for BundlePreReqCombinedStatusHelper.
#tests iOS Client

diff --git a/Engine/Source/Runtime/InstallBundleManager/Private/BundlePrereqCombinedStatusHelper.cpp b/Engine/Source/Runtime/InstallBundleManager/Private/BundlePrereqCombinedStatusHelper.cpp
--- a/Engine/Source/Runtime/InstallBundleManager/Private/BundlePrereqCombinedStatusHelper.cpp
+++ b/Engine/Source/Runtime/InstallBundleManager/Private/BundlePrereqCombinedStatusHelper.cpp
@@ -5,12 +5,12 @@
 #include "InstallBundleManagerPrivatePCH.h"
 
 FBundlePrereqCombinedStatusHelper::FBundlePrereqCombinedStatusHelper()
-: RequiredBundleNames()
+: DownloadWeight(1.f)
+, InstallWeight(1.f)
+, RequiredBundleNames()
 , BundleStatusCache()
 , CachedBundleWeights()
 , CurrentCombinedStatus()
-, DownloadWeight(1.f)
-, InstallWeight(1.f)
 , bBundleNeedsUpdate(false)
 , InstallBundleManager(nullptr)
 , TickHandle()
@@ -23,13 +23,69 @@ FBundlePrereqCombinedStatusHelper::~FBundlePrereqCombinedStatusHelper()
 	CleanUpDelegates();
 }
 
-void FBundlePrereqCombinedStatusHelper::SetupDelegates()
+FBundlePrereqCombinedStatusHelper::FBundlePrereqCombinedStatusHelper(const FBundlePrereqCombinedStatusHelper& Other)
 {
-	IPlatformInstallBundleManager::InstallBundleCompleteDelegate.AddRaw(this, &FBundlePrereqCombinedStatusHelper::OnBundleInstallComplete);
-	if (!TickHandle.IsValid())
+	*this = Other;
+}
+
+FBundlePrereqCombinedStatusHelper::FBundlePrereqCombinedStatusHelper(FBundlePrereqCombinedStatusHelper&& Other)
+{
+	*this = MoveTemp(Other);
+}
+
+FBundlePrereqCombinedStatusHelper& FBundlePrereqCombinedStatusHelper::operator=(const FBundlePrereqCombinedStatusHelper& Other)
+{
+	if (this != &Other)
+	{
+		//Just copy all this data
+		DownloadWeight = Other.DownloadWeight;
+		InstallWeight = Other.InstallWeight;
+		RequiredBundleNames = Other.RequiredBundleNames;
+		BundleStatusCache = Other.BundleStatusCache;
+		CachedBundleWeights = Other.CachedBundleWeights;
+		CurrentCombinedStatus = Other.CurrentCombinedStatus;
+		bBundleNeedsUpdate = Other.bBundleNeedsUpdate;
+		InstallBundleManager = Other.InstallBundleManager;
+		
+		//Don't copy TickHandle as we want to setup our own here
+		SetupDelegates();
+	}
+	
+	return *this;
+}
+
+FBundlePrereqCombinedStatusHelper& FBundlePrereqCombinedStatusHelper::operator=(FBundlePrereqCombinedStatusHelper&& Other)
+{
+	if (this != &Other)
 	{
-		TickHandle = FTicker::GetCoreTicker().AddTicker(FTickerDelegate::CreateRaw(this, &FBundlePrereqCombinedStatusHelper::Tick));
+		//Just copy small data
+		DownloadWeight = Other.DownloadWeight;
+		InstallWeight = Other.InstallWeight;
+		CurrentCombinedStatus = Other.CurrentCombinedStatus;
+		bBundleNeedsUpdate = Other.bBundleNeedsUpdate;
+		InstallBundleManager = Other.InstallBundleManager;
+
+		//Move bigger data
+		RequiredBundleNames = MoveTemp(Other.RequiredBundleNames);
+		BundleStatusCache = MoveTemp(Other.BundleStatusCache);
+		CachedBundleWeights = MoveTemp(Other.CachedBundleWeights);
+
+		//Prevent other from having callbacks now that its information is gone
+		Other.CleanUpDelegates();
+	
+		//Don't copy TickHandle as we want to setup our own here
+		SetupDelegates();
 	}
+	
+	return *this;
+}
+
+void FBundlePrereqCombinedStatusHelper::SetupDelegates()
+{
+	CleanUpDelegates();
+	
+	IPlatformInstallBundleManager::InstallBundleCompleteDelegate.AddRaw(this, &FBundlePrereqCombinedStatusHelper::OnBundleInstallComplete);
+	TickHandle = FTicker::GetCoreTicker().AddTicker(FTickerDelegate::CreateRaw(this, &FBundlePrereqCombinedStatusHelper::Tick));
 }
 
 void FBundlePrereqCombinedStatusHelper::CleanUpDelegates()
diff --git a/Engine/Source/Runtime/InstallBundleManager/Private/InstallBundleManagerModule.cpp b/Engine/Source/Runtime/InstallBundleManager/Private/InstallBundleManagerModule.cpp
--- a/Engine/Source/Runtime/InstallBundleManager/Private/InstallBundleManagerModule.cpp
+++ b/Engine/Source/Runtime/InstallBundleManager/Private/InstallBundleManagerModule.cpp
@@ -2,3 +2,4 @@
 #include "InstallBundleManagerModule.h"
 
 IMPLEMENT_MODULE(FInstallBundleManagerModule, InstallBundleManager);
+DEFINE_LOG_CATEGORY(LogInstallBundleManager);
diff --git a/Engine/Source/Runtime/InstallBundleManager/Private/InstallBundleManagerPrivatePCH.h b/Engine/Source/Runtime/InstallBundleManager/Private/InstallBundleManagerPrivatePCH.h
--- a/Engine/Source/Runtime/InstallBundleManager/Private/InstallBundleManagerPrivatePCH.h
+++ b/Engine/Source/Runtime/InstallBundleManager/Private/InstallBundleManagerPrivatePCH.h
@@ -4,3 +4,5 @@
 
 #include "CoreMinimal.h"
 #include "GenericPlatform/GenericPlatformInstallBundleManager.h"
+
+DECLARE_LOG_CATEGORY_EXTERN(LogInstallBundleManager, Log, All);
diff --git a/Engine/Source/Runtime/InstallBundleManager/Public/BundlePrereqCombinedStatusHelper.h b/Engine/Source/Runtime/InstallBundleManager/Public/BundlePrereqCombinedStatusHelper.h
--- a/Engine/Source/Runtime/InstallBundleManager/Public/BundlePrereqCombinedStatusHelper.h
+++ b/Engine/Source/Runtime/InstallBundleManager/Public/BundlePrereqCombinedStatusHelper.h
@@ -42,12 +42,23 @@ public:
 	FBundlePrereqCombinedStatusHelper();
 	~FBundlePrereqCombinedStatusHelper();
 	
+	FBundlePrereqCombinedStatusHelper(const FBundlePrereqCombinedStatusHelper& Other);
+	FBundlePrereqCombinedStatusHelper(FBundlePrereqCombinedStatusHelper&& Other);
+	
+	FBundlePrereqCombinedStatusHelper& operator=(const FBundlePrereqCombinedStatusHelper& Other);
+	FBundlePrereqCombinedStatusHelper& operator=(FBundlePrereqCombinedStatusHelper&& Other);
+	
 	//Setup tracking for all bundles required in the supplied BundleContentState
 	void SetBundlesToTrackFromContentState(FInstallBundleContentState& BundleContentState);
 	
 	//Get current CombinedBundleStatus for everything setup to track
 	const FCombinedBundleStatus& GetCurrentCombinedState() const;
 	
+	//How to weight downloads vs. installs. Defaults to even. Does not have to add up to 1.0.
+	//Setting Download to .5 and Install to .5 will be the same as setting Download to 1.f and Install to 1.f.
+	float DownloadWeight;
+	float InstallWeight;
+	
 private:
 	bool Tick(float dt);
 	void UpdateBundleCache();
@@ -74,10 +85,6 @@ private:
 	
 	FCombinedBundleStatus CurrentCombinedStatus;
 	
-	//How to weight downloads vs. installs
-	float DownloadWeight;
-	float InstallWeight;
-	
 	bool bBundleNeedsUpdate;
 	
 	IPlatformInstallBundleManager* InstallBundleManager;
