Fixing Missing Header in non-unity build.
#tests iOS Client compile

diff --git a/Engine/Source/Runtime/InstallBundleManager/Private/InstallBundleManagerModule.cpp b/Engine/Source/Runtime/InstallBundleManager/Private/InstallBundleManagerModule.cpp
--- a/Engine/Source/Runtime/InstallBundleManager/Private/InstallBundleManagerModule.cpp
+++ b/Engine/Source/Runtime/InstallBundleManager/Private/InstallBundleManagerModule.cpp
@@ -1,4 +1,5 @@
 // Copyright 1998-2019 Epic Games, Inc. All Rights Reserved.
+#include "InstallBundleManagerPrivatePCH.h"
 #include "InstallBundleManagerModule.h"
 
 IMPLEMENT_MODULE(FInstallBundleManagerModule, InstallBundleManager);
