BundlePrereqCombindStatusHelper - fix shader caching on Android being combined with other progress.

diff --git a/Engine/Source/Runtime/InstallBundleManager/Private/BundlePrereqCombinedStatusHelper.cpp b/Engine/Source/Runtime/InstallBundleManager/Private/BundlePrereqCombinedStatusHelper.cpp
--- a/Engine/Source/Runtime/InstallBundleManager/Private/BundlePrereqCombinedStatusHelper.cpp
+++ b/Engine/Source/Runtime/InstallBundleManager/Private/BundlePrereqCombinedStatusHelper.cpp
@@ -153,6 +153,7 @@ void FBundlePrereqCombinedStatusHelper::UpdateCombinedStatus()
 	
 	EInstallBundleStatus EarliestBundleState = EInstallBundleStatus::Count;
 	bool bIsAnythingPaused = false;
+	bool bIsAnythingFinishing = false;
 	
 	//if we don't yet have a bundle status cache entry for a particular requirement
 	//then we can't yet tell what work is required on that bundle yet. We need to go ahead and make sure we don't
@@ -164,14 +165,21 @@ void FBundlePrereqCombinedStatusHelper::UpdateCombinedStatus()
 		EarliestBundleState = EInstallBundleStatus::Downloading;
 	}
 	
+	float EarliestFinishingPercent = 1.0f;
 	for (const TPair<FName,FInstallBundleStatus>& BundlePair : BundleStatusCache)
 	{
 		if (BundlePair.Value.Status < EarliestBundleState)
 		{
 			EarliestBundleState = BundlePair.Value.Status;
 		}
+
+		if (!bIsAnythingFinishing && BundlePair.Value.Status == EInstallBundleStatus::Finishing)
+		{
+			EarliestFinishingPercent = BundlePair.Value.Finishing_Percent;
+			bIsAnythingFinishing = true;
+		}
 		
-		bIsAnythingPaused = (bIsAnythingPaused || (BundlePair.Value.PauseFlags != EInstallBundlePauseFlags::None));
+		bIsAnythingPaused = bIsAnythingPaused || BundlePair.Value.PauseFlags != EInstallBundlePauseFlags::None;
 	}
 	
 	//if we have any paused bundles, and we have any bundle that isn't finished installed, we are Paused
@@ -186,14 +194,21 @@ void FBundlePrereqCombinedStatusHelper::UpdateCombinedStatus()
 	{
 		CurrentCombinedStatus.CombinedState = FCombinedBundleStatus::ECombinedBundleStateEnum::Initializing;
 	}
-	else if (EarliestBundleState < EInstallBundleStatus::Finishing)
+	else if (EarliestBundleState <= EInstallBundleStatus::Installing)
 	{
 		CurrentCombinedStatus.CombinedState = FCombinedBundleStatus::ECombinedBundleStateEnum::Updating;
 	}
-	else if (EarliestBundleState < EInstallBundleStatus::Installed)
+	else if (EarliestBundleState <= EInstallBundleStatus::Finishing)
 	{
-		CurrentCombinedStatus.CombinedState = FCombinedBundleStatus::ECombinedBundleStateEnum::Finishing;
-		CurrentCombinedStatus.bDoesCurrentStateSupportPausing = false;
+		if (bIsAnythingFinishing)
+		{
+			CurrentCombinedStatus.CombinedState = FCombinedBundleStatus::ECombinedBundleStateEnum::Finishing;
+			CurrentCombinedStatus.ProgressPercent = EarliestFinishingPercent;
+		}
+		else
+		{
+			CurrentCombinedStatus.CombinedState = FCombinedBundleStatus::ECombinedBundleStateEnum::Updating;
+		}
 	}
 	else if (EarliestBundleState == EInstallBundleStatus::Installed)
 	{
