Fix for Android progress sliding backwards early in the percent.
#tests Android Classic Client

diff --git a/Engine/Source/Runtime/InstallBundleManager/Private/BundlePrereqCombinedStatusHelper.cpp b/Engine/Source/Runtime/InstallBundleManager/Private/BundlePrereqCombinedStatusHelper.cpp
--- a/Engine/Source/Runtime/InstallBundleManager/Private/BundlePrereqCombinedStatusHelper.cpp
+++ b/Engine/Source/Runtime/InstallBundleManager/Private/BundlePrereqCombinedStatusHelper.cpp
@@ -200,6 +200,8 @@ void FBundlePrereqCombinedStatusHelper::UpdateCombinedStatus()
 	}
 	else if (EarliestBundleState <= EInstallBundleStatus::Finishing)
 	{
+		//Handles the case where one of our Bundles was finishing and we have finished everything else.
+		//Now just shows our earliest bundle that is finishing.
 		if (bIsAnythingFinishing)
 		{
 			CurrentCombinedStatus.CombinedState = FCombinedBundleStatus::ECombinedBundleStateEnum::Finishing;
@@ -245,16 +247,11 @@ float FBundlePrereqCombinedStatusHelper::GetIndividualWeightedProgressPercent(FI
 	const float TotalWeight = DownloadWeight + InstallWeight;
 	float CombinedOverallProgressPercent = 0.f;
 	
-	//Completed is maximum progress possible
-	if (BundleStatus.Status >= EInstallBundleStatus::Installed)
+	//If we are done installing then lets just coung our progress as 1.0f. Finishing progress is handled in UpdateCombinedStatus
+	if (BundleStatus.Status > EInstallBundleStatus::Installing)
 	{
 		CombinedOverallProgressPercent = 1.0f;
 	}
-	//Once we are in finishing, we display a new bar for this step, so just show raw Finishing_Percent
-	else if (BundleStatus.Status >= EInstallBundleStatus::Finishing)
-	{
-		CombinedOverallProgressPercent = BundleStatus.Finishing_Percent;
-	}
 	else
 	{
 		if (ensureAlwaysMsgf((TotalWeight > 0), TEXT("Invalid Weights for FBundleStatusTracker! Need to have weights > 0.f!")))
