Changing Initializing State to move to Updating once we begin downloads instead of waiting for progress.
#tests iOS Client

diff --git a/Engine/Source/Runtime/InstallBundleManager/Private/BundlePrereqCombinedStatusHelper.cpp b/Engine/Source/Runtime/InstallBundleManager/Private/BundlePrereqCombinedStatusHelper.cpp
--- a/Engine/Source/Runtime/InstallBundleManager/Private/BundlePrereqCombinedStatusHelper.cpp
+++ b/Engine/Source/Runtime/InstallBundleManager/Private/BundlePrereqCombinedStatusHelper.cpp
@@ -190,7 +190,7 @@ void FBundlePrereqCombinedStatusHelper::UpdateCombinedStatus()
 	//Otherwise start with True and override those specific cases bellow
 	CurrentCombinedStatus.bDoesCurrentStateSupportPausing = bBundleNeedsUpdate;
 	
-	if (CurrentCombinedStatus.ProgressPercent == 0.f)
+	if ((EarliestBundleState == EInstallBundleStatus::Requested) || (EarliestBundleState == EInstallBundleStatus::Count))
 	{
 		CurrentCombinedStatus.CombinedState = FCombinedBundleStatus::ECombinedBundleStateEnum::Initializing;
 	}
