Remove unnecessary check for InstallDownloadProgress and fix weight when there is no background download progress.

diff --git a/Engine/Source/Runtime/InstallBundleManager/Private/BundlePrereqCombinedStatusHelper.cpp b/Engine/Source/Runtime/InstallBundleManager/Private/BundlePrereqCombinedStatusHelper.cpp
--- a/Engine/Source/Runtime/InstallBundleManager/Private/BundlePrereqCombinedStatusHelper.cpp
+++ b/Engine/Source/Runtime/InstallBundleManager/Private/BundlePrereqCombinedStatusHelper.cpp
@@ -281,12 +281,10 @@ float FBundlePrereqCombinedStatusHelper::GetIndividualWeightedProgressPercent(FI
 			}
 			
 			float InstallWeightedPercent = 0.f;
-			if (BundleStatus.InstallDownloadProgress.IsSet())
-			{
-				//If we didn't have background download progress, just use our InstallWeight as we don't have Download Progress to use
-				const float AppliedInstallWeight = bDidHaveBackgroundDownloadProgress ? (InstallWeight / TotalWeight) : InstallWeight;
-				InstallWeightedPercent = BundleStatus.Install_Percent * AppliedInstallWeight;
-			}
+			
+			//If we didn't have background download progress, just use our InstallWeight as we don't have Download Progress to use
+			const float AppliedInstallWeight = bDidHaveBackgroundDownloadProgress ? (InstallWeight / TotalWeight) : 1.0f;
+			InstallWeightedPercent = BundleStatus.Install_Percent * AppliedInstallWeight;
 			
 			CombinedOverallProgressPercent = DownloadWeightedPercent + InstallWeightedPercent;
 			FMath::Clamp(CombinedOverallProgressPercent, 0.f, 1.0f);
