Move Install Bundle Manager interface out of core and into the bundle manager module.

diff --git a/Engine/Source/Runtime/Core/Private/GenericPlatform/GenericPlatformInstallBundleManager.cpp b/Engine/Source/Runtime/Core/Private/GenericPlatform/GenericPlatformInstallBundleManager.cpp
deleted file mode 100644
diff --git a/Engine/Source/Runtime/Core/Private/GenericPlatform/GenericPlatformMisc.cpp b/Engine/Source/Runtime/Core/Private/GenericPlatform/GenericPlatformMisc.cpp
--- a/Engine/Source/Runtime/Core/Private/GenericPlatform/GenericPlatformMisc.cpp
+++ b/Engine/Source/Runtime/Core/Private/GenericPlatform/GenericPlatformMisc.cpp
@@ -26,7 +26,6 @@
 #include "HAL/ExceptionHandling.h"
 #include "GenericPlatform/GenericPlatformCrashContext.h"
 #include "GenericPlatform/GenericPlatformDriver.h"
-#include "GenericPlatform/GenericPlatformInstallBundleManager.h"
 #include "ProfilingDebugging/ExternalProfiler.h"
 #include "HAL/LowLevelMemTracker.h"
 #include "Templates/Function.h"
@@ -825,36 +824,6 @@ IPlatformChunkInstall* FGenericPlatformMisc::GetPlatformChunkInstall()
 	return &Singleton;
 }
 
-IPlatformInstallBundleManager* FGenericPlatformMisc::GetPlatformInstallBundleManager()
-{
-	static IPlatformInstallBundleManager* Manager = nullptr;
-	static bool bCheckedIni = false;
-
-	if (Manager)
-		return Manager;
-
-	if (!bCheckedIni && !GEngineIni.IsEmpty())
-	{
-		FString ModuleName;
-		IPlatformInstallBundleManagerModule* Module = nullptr;
-		GConfig->GetString(TEXT("InstallBundleManager"), TEXT("ModuleName"), ModuleName, GEngineIni);
-
-		if (FModuleManager::Get().ModuleExists(*ModuleName))
-		{
-			FModuleStatus Status;
-			Module = FModuleManager::LoadModulePtr<IPlatformInstallBundleManagerModule>(*ModuleName);
-			if (Module)
-			{
-				Manager = Module->GetInstallBundleManager();
-			}
-		}
-
-		bCheckedIni = true;
-	}
-
-	return Manager;
-}
-
 void GenericPlatformMisc_GetProjectFilePathProjectDir(FString& OutGameDir)
 {
 	// Here we derive the game path from the project file location.
diff --git a/Engine/Source/Runtime/Core/Public/GenericPlatform/GenericPlatformInstallBundleManager.h b/Engine/Source/Runtime/Core/Public/GenericPlatform/GenericPlatformInstallBundleManager.h
deleted file mode 100644
diff --git a/Engine/Source/Runtime/Core/Public/GenericPlatform/GenericPlatformMisc.h b/Engine/Source/Runtime/Core/Public/GenericPlatform/GenericPlatformMisc.h
--- a/Engine/Source/Runtime/Core/Public/GenericPlatform/GenericPlatformMisc.h
+++ b/Engine/Source/Runtime/Core/Public/GenericPlatform/GenericPlatformMisc.h
@@ -10,7 +10,7 @@
 class Error;
 class GenericApplication;
 class IPlatformChunkInstall;
-class IPlatformInstallBundleManager;
+class IInstallBundleManager;
 class IPlatformCompression;
 struct FGenericCrashContext;
 struct FGenericMemoryWarningContext;
@@ -855,13 +855,6 @@ public:
 	 */
 	static IPlatformChunkInstall* GetPlatformChunkInstall();
 
-	/**
-	 * Returns the platform specific Install Bundle Manager
-	 *
-	 * @return	Returns the platform specific Install Bundle Manager implementation
-	 */
-	static IPlatformInstallBundleManager* GetPlatformInstallBundleManager();
-
 	/**
 	 * Returns the platform specific compression interface
 	 *
diff --git a/Engine/Source/Runtime/InstallBundleManager/Private/BundlePrereqCombinedStatusHelper.cpp b/Engine/Source/Runtime/InstallBundleManager/Private/BundlePrereqCombinedStatusHelper.cpp
--- a/Engine/Source/Runtime/InstallBundleManager/Private/BundlePrereqCombinedStatusHelper.cpp
+++ b/Engine/Source/Runtime/InstallBundleManager/Private/BundlePrereqCombinedStatusHelper.cpp
@@ -84,13 +84,13 @@ void FBundlePrereqCombinedStatusHelper::SetupDelegates()
 {
 	CleanUpDelegates();
 	
-	IPlatformInstallBundleManager::InstallBundleCompleteDelegate.AddRaw(this, &FBundlePrereqCombinedStatusHelper::OnBundleInstallComplete);
+	IInstallBundleManager::InstallBundleCompleteDelegate.AddRaw(this, &FBundlePrereqCombinedStatusHelper::OnBundleInstallComplete);
 	TickHandle = FTicker::GetCoreTicker().AddTicker(FTickerDelegate::CreateRaw(this, &FBundlePrereqCombinedStatusHelper::Tick));
 }
 
 void FBundlePrereqCombinedStatusHelper::CleanUpDelegates()
 {
-	IPlatformInstallBundleManager::InstallBundleCompleteDelegate.RemoveAll(this);
+	IInstallBundleManager::InstallBundleCompleteDelegate.RemoveAll(this);
 	if (TickHandle.IsValid())
 	{
 		FTicker::GetCoreTicker().RemoveTicker(TickHandle);
@@ -128,7 +128,7 @@ void FBundlePrereqCombinedStatusHelper::UpdateBundleCache()
 	//if we haven't already set this up, lets try to set it now
 	if (nullptr == InstallBundleManager)
 	{
-		InstallBundleManager = FPlatformMisc::GetPlatformInstallBundleManager();
+		InstallBundleManager = IInstallBundleManager::GetPlatformInstallBundleManager();
 	}
 	
 	if (ensureAlwaysMsgf((nullptr != InstallBundleManager), TEXT("Invalid InstallBundleManager during UpdateBundleCache! Needs to be valid during run!")))
diff --git a/Engine/Source/Runtime/InstallBundleManager/Private/InstallBundleManagerInterface.cpp b/Engine/Source/Runtime/InstallBundleManager/Private/InstallBundleManagerInterface.cpp
new file mode 100644
--- /dev/null
+++ b/Engine/Source/Runtime/InstallBundleManager/Private/InstallBundleManagerInterface.cpp
@@ -0,0 +1,41 @@
+// Copyright 1998-2019 Epic Games, Inc. All Rights Reserved.
+
+#include "InstallBundleManagerInterface.h"
+#include "InstallBundleManagerModule.h"
+
+FInstallBundleCompleteMultiDelegate IInstallBundleManager::InstallBundleCompleteDelegate;
+
+FInstallBundleCompleteMultiDelegate IInstallBundleManager::RemoveBundleCompleteDelegate;
+
+FInstallBundlePausedMultiDelegate IInstallBundleManager::PausedBundleDelegate;
+
+IInstallBundleManager* IInstallBundleManager::GetPlatformInstallBundleManager()
+{
+	static IInstallBundleManager* Manager = nullptr;
+	static bool bCheckedIni = false;
+
+	if (Manager)
+		return Manager;
+
+	if (!bCheckedIni && !GEngineIni.IsEmpty())
+	{
+		FString ModuleName;
+		IInstallBundleManagerModule* Module = nullptr;
+		GConfig->GetString(TEXT("InstallBundleManager"), TEXT("ModuleName"), ModuleName, GEngineIni);
+
+		if (FModuleManager::Get().ModuleExists(*ModuleName))
+		{
+			FModuleStatus Status;
+			Module = FModuleManager::LoadModulePtr<IInstallBundleManagerModule>(*ModuleName);
+			if (Module)
+			{
+				Manager = Module->GetInstallBundleManager();
+			}
+		}
+
+		bCheckedIni = true;
+	}
+
+	return Manager;
+}
+
diff --git a/Engine/Source/Runtime/InstallBundleManager/Private/InstallBundleManagerPrivatePCH.h b/Engine/Source/Runtime/InstallBundleManager/Private/InstallBundleManagerPrivatePCH.h
--- a/Engine/Source/Runtime/InstallBundleManager/Private/InstallBundleManagerPrivatePCH.h
+++ b/Engine/Source/Runtime/InstallBundleManager/Private/InstallBundleManagerPrivatePCH.h
@@ -3,6 +3,6 @@
 #pragma once
 
 #include "CoreMinimal.h"
-#include "GenericPlatform/GenericPlatformInstallBundleManager.h"
+#include "InstallBundleManagerInterface.h"
 
 DECLARE_LOG_CATEGORY_EXTERN(LogInstallBundleManager, Log, All);
diff --git a/Engine/Source/Runtime/InstallBundleManager/Public/BundlePrereqCombinedStatusHelper.h b/Engine/Source/Runtime/InstallBundleManager/Public/BundlePrereqCombinedStatusHelper.h
--- a/Engine/Source/Runtime/InstallBundleManager/Public/BundlePrereqCombinedStatusHelper.h
+++ b/Engine/Source/Runtime/InstallBundleManager/Public/BundlePrereqCombinedStatusHelper.h
@@ -2,7 +2,7 @@
 #pragma once
 
 #include "CoreMinimal.h"
-#include "GenericPlatform/GenericPlatformInstallBundleManager.h"
+#include "InstallBundleManagerInterface.h"
 
 //Handles calculating the bundle status by combining progress from all of its
 //Prerequisites. Allows you to display one progress percent that is weighted based on all
@@ -88,6 +88,6 @@ private:
 	
 	bool bBundleNeedsUpdate;
 	
-	IPlatformInstallBundleManager* InstallBundleManager;
+	IInstallBundleManager* InstallBundleManager;
 	FDelegateHandle TickHandle;
 };
diff --git a/Engine/Source/Runtime/InstallBundleManager/Public/InstallBundleManagerInterface.h b/Engine/Source/Runtime/InstallBundleManager/Public/InstallBundleManagerInterface.h
new file mode 100644
--- /dev/null
+++ b/Engine/Source/Runtime/InstallBundleManager/Public/InstallBundleManagerInterface.h
@@ -0,0 +1,340 @@
+// Copyright 1998-2019 Epic Games, Inc. All Rights Reserved.
+
+#pragma once
+
+#include "Misc/EnumClassFlags.h"
+#include "Misc/ConfigCacheIni.h"
+#include "Logging/LogMacros.h"
+#include "Internationalization/Text.h"
+
+class IAnalyticsProviderET;
+
+enum class EInstallBundleModuleInitResult : int
+{
+	OK,
+	BuildMetaDataNotFound,
+	BuildMetaDataParsingError,
+	DistributionRootParseError,
+	DistributionRootDownloadError,
+	ManifestArchiveError,
+	ManifestCreationError,
+	ManifestDownloadError,
+	BackgroundDownloadsIniDownloadError,
+	NoInternetConnectionError,
+	ConfigurationError,
+	Count
+};
+
+inline const TCHAR* LexToString(EInstallBundleModuleInitResult Result)
+{
+	using UnderType = __underlying_type(EInstallBundleModuleInitResult);
+	static const TCHAR* Strings[] =
+	{
+		TEXT("OK"),
+		TEXT("BuildMetaDataNotFound"),
+		TEXT("BuildMetaDataParsingError"),
+		TEXT("DistributionRootParseError"),
+		TEXT("DistributionRootDownloadError"),
+		TEXT("ManifestArchiveError"),
+		TEXT("ManifestCreationError"),
+		TEXT("ManifestDownloadError"),
+		TEXT("BackgroundDownloadsIniDownloadError"),
+		TEXT("NoInternetConnectionError"),
+		TEXT("ConfigurationError"),
+	};
+	static_assert(static_cast<UnderType>(EInstallBundleModuleInitResult::Count) == ARRAY_COUNT(Strings), "");
+
+	return Strings[static_cast<UnderType>(Result)];
+}
+
+enum class EInstallBundleResult : int
+{
+	OK,
+	FailedPrereqRequiresLatestClient,
+	InstallError,
+	InstallerOutOfDiskSpaceError,
+	ManifestArchiveError,
+	UserCancelledError,
+	InitializationError,
+	Count,
+};
+
+inline const TCHAR* LexToString(EInstallBundleResult Result)
+{
+	using UnderType = __underlying_type(EInstallBundleResult);
+	static const TCHAR* Strings[] =
+	{
+		TEXT("OK"),
+		TEXT("FailedPrereqRequiresLatestClient"),
+		TEXT("InstallError"),
+		TEXT("InstallerOutOfDiskSpaceError"),
+		TEXT("ManifestArchiveError"),
+		TEXT("UserCancelledError"),
+		TEXT("InitializationError"),
+	};
+	static_assert(static_cast<UnderType>(EInstallBundleResult::Count) == ARRAY_COUNT(Strings), "");
+
+	return Strings[static_cast<UnderType>(Result)];
+}
+
+enum class EInstallBundlePauseFlags : uint32
+{
+	None = 0,
+	OnCellularNetwork = (1 << 0),
+	NoInternetConnection = (1 << 1),
+	UserPaused	= (1 << 2)
+};
+ENUM_CLASS_FLAGS(EInstallBundlePauseFlags)
+
+inline const TCHAR* GetInstallBundlePauseReason(EInstallBundlePauseFlags Flags)
+{
+	// Return the most appropriate reason given the flags
+
+	if (EnumHasAnyFlags(Flags, EInstallBundlePauseFlags::UserPaused))
+		return TEXT("UserPaused");
+
+	if (EnumHasAnyFlags(Flags, EInstallBundlePauseFlags::NoInternetConnection))
+		return TEXT("NoInternetConnection");
+	
+	if (EnumHasAnyFlags(Flags, EInstallBundlePauseFlags::OnCellularNetwork))
+		return TEXT("OnCellularNetwork");
+
+	return TEXT("");
+}
+
+enum class EInstallBundleRequestFlags : uint32
+{
+	None											= 0,
+	CheckForCellularDataUsage						= (1 << 0),
+	UseBackgroundDownloads							= (1 << 1),
+	SendNotificationIfDownloadCompletesInBackground = (1 << 2),
+	ForceNoPatching									= (1 << 3),
+	Defaults = UseBackgroundDownloads,
+};
+ENUM_CLASS_FLAGS(EInstallBundleRequestFlags)
+
+enum class EInstallBundleStatus : int
+{
+	Requested,
+	QueuedForDownload,
+	Downloading,
+	QueuedForInstall,
+	Installing,
+	QueuedForFinish,
+	Finishing,
+	Installed,
+	Count,
+};
+
+inline const TCHAR* LexToString(EInstallBundleStatus Status)
+{
+	using UnderType = __underlying_type(EInstallBundleStatus);
+	static const TCHAR* Strings[] =
+	{
+		TEXT("Requested"),
+		TEXT("QueuedForDownload"),
+		TEXT("Downloading"),
+		TEXT("QueuedForInstall"),
+		TEXT("Installing"),
+		TEXT("QueuedForFinish"),
+		TEXT("Finishing"),
+		TEXT("Installed"),
+	};
+
+	static_assert(static_cast<UnderType>(EInstallBundleStatus::Count) == ARRAY_COUNT(Strings), "");
+
+	return Strings[static_cast<UnderType>(Status)];
+}
+
+struct FInstallBundleDownloadProgress
+{
+	// Num bytes received
+	uint64 BytesDownloaded = 0;
+	// Num bytes written to storage (<= BytesDownloaded)
+	uint64 BytesDownloadedAndWritten = 0;
+	// Num bytes needed
+	uint64 TotalBytesToDownload = 0;
+	// Num bytes that failed to download
+	uint64 TotalBytesFailedToDownload = 0;
+	float PercentComplete = 0;
+};
+
+struct FInstallBundleStatus
+{
+	FName BundleName;
+
+	EInstallBundleStatus Status = EInstallBundleStatus::QueuedForDownload;
+
+	EInstallBundlePauseFlags PauseFlags = EInstallBundlePauseFlags::None;
+
+	FText StatusText;
+
+	// Download progress of EInstallBundleStatus::Downloading
+	// Will be set if Status >= EInstallBundleStatus::Downloading
+	TOptional<FInstallBundleDownloadProgress> BackgroundDownloadProgress;
+
+	// Download progress of EInstallBundleStatus::Install
+	// Will be set if Status >= EInstallBundleStatus::Install
+	// We may download during install if background downloads are turned off or fail
+	// We may also do small downloads during install as a normal part of installation
+	TOptional<FInstallBundleDownloadProgress> InstallDownloadProgress;
+
+	float Install_Percent = 0;
+
+	float Finishing_Percent = 0;
+};
+
+struct FInstallBundleResultInfo
+{
+	FName BundleName;
+	EInstallBundleResult Result = EInstallBundleResult::OK;
+	bool bIsStartup = false;
+
+	// Currently, these just forward BPT Error info
+	FText OptionalErrorText;
+	FString OptionalErrorCode;
+};
+
+struct FInstallBundlePauseInfo
+{
+	FName BundleName;
+	EInstallBundlePauseFlags PauseFlags = EInstallBundlePauseFlags::None;
+};
+
+enum class EInstallBundleContentState : int
+{
+	InitializationError,
+	NotInstalled,
+	NeedsUpdate,
+	UpToDate,
+	Count,
+};
+inline const TCHAR* LexToString(EInstallBundleContentState State)
+{
+	using UnderType = __underlying_type(EInstallBundleContentState);
+	static const TCHAR* Strings[] =
+	{
+		TEXT("InitializationError"),
+		TEXT("NotInstalled"),
+		TEXT("NeedsUpdate"),
+		TEXT("UpToDate"),
+	};
+	static_assert(static_cast<UnderType>(EInstallBundleContentState::Count) == ARRAY_COUNT(Strings), "");
+
+	return Strings[static_cast<UnderType>(State)];
+}
+
+struct FInstallBundleContentState
+{
+	EInstallBundleContentState State = EInstallBundleContentState::InitializationError;
+	TMap<FName, EInstallBundleContentState> IndividualBundleStates;
+	TMap<FName, float> IndividualBundleWeights;
+	uint64 DownloadSize = 0;
+	uint64 InstallSize = 0;
+	uint64 FreeSpace = 0;
+};
+
+enum class EInstallBundleGetContentStateFlags : uint32
+{
+	None = 0,
+	ForceNoPatching = (1 << 0),
+};
+ENUM_CLASS_FLAGS(EInstallBundleGetContentStateFlags);
+
+enum class EInstallBundleRequestInfoFlags : int32
+{
+	None							= 0,
+	EnqueuedBundlesForInstall		= (1 << 0),
+	EnqueuedBundlesForRemoval		= (1 << 1),
+	SkippedAlreadyMountedBundles	= (1 << 2),
+	SkippedBundlesQueuedForRemoval	= (1 << 3),
+	SkippedBundlesQueuedForInstall  = (1 << 4), // Only valid for removal requests
+	SkippedUnknownBundles			= (1 << 5),
+	InitializationError				= (1 << 6), // Can't enqueue because the bundle manager failed to initialize
+};
+ENUM_CLASS_FLAGS(EInstallBundleRequestInfoFlags);
+
+struct FInstallBundleRequestInfo
+{
+	EInstallBundleRequestInfoFlags InfoFlags = EInstallBundleRequestInfoFlags::None;
+	TArray<FName> BundlesQueuedForInstall;
+	TArray<FName> BundlesQueuedForRemoval;
+};
+
+enum class EInstallBundleCancelFlags : int32
+{
+	None		= 0,
+	Resumable	= (1 << 0),
+};
+ENUM_CLASS_FLAGS(EInstallBundleCancelFlags);
+
+enum class EInstallBundleManagerInitErrorHandlerResult
+{
+	NotHandled, // Defer to the next handler
+	Retry, // Try to initialize again
+	StopInitialization, // Stop trying to initialize
+};
+
+DECLARE_DELEGATE_RetVal_OneParam(EInstallBundleManagerInitErrorHandlerResult, FInstallBundleManagerInitErrorHandler, EInstallBundleModuleInitResult);
+
+DECLARE_MULTICAST_DELEGATE_OneParam(FInstallBundleCompleteMultiDelegate, FInstallBundleResultInfo);
+DECLARE_MULTICAST_DELEGATE_OneParam(FInstallBundlePausedMultiDelegate, FInstallBundlePauseInfo);
+
+DECLARE_DELEGATE_OneParam(FInstallBundleGetContentStateDelegate, FInstallBundleContentState);
+
+class INSTALLBUNDLEMANAGER_API IInstallBundleManager
+{
+public:
+	static FInstallBundleCompleteMultiDelegate InstallBundleCompleteDelegate;
+	static FInstallBundleCompleteMultiDelegate RemoveBundleCompleteDelegate;
+	static FInstallBundlePausedMultiDelegate PausedBundleDelegate;
+
+	static IInstallBundleManager* GetPlatformInstallBundleManager();
+
+	virtual ~IInstallBundleManager() {}
+
+	virtual bool HasBuildMetaData() const = 0;
+
+	virtual void PushInitErrorCallback(FInstallBundleManagerInitErrorHandler Callback) = 0;
+	virtual void PopInitErrorCallback() = 0;
+
+	virtual bool IsInitialized() const = 0;
+	virtual bool IsInitializing() const = 0;
+
+	virtual bool IsActive() const = 0;
+
+	virtual FInstallBundleRequestInfo RequestUpdateContent(FName BundleName, EInstallBundleRequestFlags Flags) = 0;
+	virtual FInstallBundleRequestInfo RequestUpdateContent(TArrayView<FName> BundleNames, EInstallBundleRequestFlags Flags) = 0;
+
+	virtual void GetContentState(FName BundleName, EInstallBundleGetContentStateFlags Flags, bool bAddDependencies, FInstallBundleGetContentStateDelegate Callback, FName RequestTag = TEXT("None")) = 0;
+	virtual void GetContentState(TArrayView<FName> BundleNames, EInstallBundleGetContentStateFlags Flags, bool bAddDependencies, FInstallBundleGetContentStateDelegate Callback, FName RequestTag = TEXT("None")) = 0;
+
+    virtual void CancelAllGetContentStateRequestsForTag(FName RequestTag) = 0;
+    
+	virtual FInstallBundleRequestInfo RequestRemoveBundle(FName BundleName) = 0;
+
+	virtual void RequestRemoveBundleOnNextInit(FName BundleName) = 0;
+
+	virtual void CancelRequestRemoveBundleOnNextInit(FName BundleName) = 0;
+
+	virtual void CancelBundle(FName BundleName, EInstallBundleCancelFlags Flags) = 0;
+
+	virtual void CancelAllBundles(EInstallBundleCancelFlags Flags) = 0;
+
+	virtual bool PauseBundle(FName BundleName) = 0;
+
+	virtual void ResumeBundle(FName BundleName) = 0;
+
+	virtual void RequestPausedBundleCallback() const = 0;
+
+	virtual TOptional<FInstallBundleStatus> GetBundleProgress(FName BundleName) const = 0;
+
+	virtual void UpdateContentRequestFlags(FName BundleName, EInstallBundleRequestFlags AddFlags, EInstallBundleRequestFlags RemoveFlags) = 0;
+
+	virtual bool IsNullInterface() const = 0;
+
+	virtual void SetErrorSimulationCommands(const FString& CommandLine) {}
+
+	virtual TSharedPtr<IAnalyticsProviderET> GetAnalyticsProvider() const { return TSharedPtr<IAnalyticsProviderET>(); }
+};
+
diff --git a/Engine/Source/Runtime/InstallBundleManager/Public/InstallBundleManagerModule.h b/Engine/Source/Runtime/InstallBundleManager/Public/InstallBundleManagerModule.h
--- a/Engine/Source/Runtime/InstallBundleManager/Public/InstallBundleManagerModule.h
+++ b/Engine/Source/Runtime/InstallBundleManager/Public/InstallBundleManagerModule.h
@@ -19,3 +19,43 @@ public:
 	{
 	}
 };
+
+/**
+ * Base Module Interface for InstallBundleManager implementation modules
+ */
+class IInstallBundleManagerModule : public IModuleInterface
+{
+public:
+	virtual void PreUnloadCallback() override
+	{
+		InstallBundleManager.Reset();
+	}
+
+	IInstallBundleManager* GetInstallBundleManager()
+	{
+		return InstallBundleManager.Get();
+	}
+
+protected:
+	TUniquePtr<IInstallBundleManager> InstallBundleManager;
+};
+
+/**
+ * Module Interface for InstallBundleManager implementation modules
+ */
+template<class InstallBundleManagerModuleImpl>
+class TInstallBundleManagerModule : public IInstallBundleManagerModule
+{
+public:
+	virtual void StartupModule() override
+	{
+		// Only instantiate the bundle manager if this is the version the game has been configured to use
+		FString ModuleName;
+		GConfig->GetString(TEXT("InstallBundleManager"), TEXT("ModuleName"), ModuleName, GEngineIni);
+
+		if (FModuleManager::Get().GetModule(*ModuleName) == this)
+		{
+			InstallBundleManager = MakeUnique<InstallBundleManagerModuleImpl>();
+		}
+	}
+};
diff --git a/Engine/Source/Runtime/Launch/Launch.Build.cs b/Engine/Source/Runtime/Launch/Launch.Build.cs
--- a/Engine/Source/Runtime/Launch/Launch.Build.cs
+++ b/Engine/Source/Runtime/Launch/Launch.Build.cs
@@ -34,11 +34,12 @@ public class Launch : ModuleRules
 				"TraceLog",
 				"Overlay",
 				"UtilityShaders",
-				"PreLoadScreen"
+				"PreLoadScreen",
+				"InstallBundleManager"
 			});
 
 		// Set a macro allowing us to switch between debuggame/development configuration
-		if(Target.Configuration == UnrealTargetConfiguration.DebugGame)
+		if (Target.Configuration == UnrealTargetConfiguration.DebugGame)
 		{
 			PrivateDefinitions.Add("UE_BUILD_DEVELOPMENT_WITH_DEBUGGAME=1");
 		}
diff --git a/Engine/Source/Runtime/Launch/Private/LaunchEngineLoop.cpp b/Engine/Source/Runtime/Launch/Private/LaunchEngineLoop.cpp
--- a/Engine/Source/Runtime/Launch/Private/LaunchEngineLoop.cpp
+++ b/Engine/Source/Runtime/Launch/Private/LaunchEngineLoop.cpp
@@ -152,6 +152,7 @@
 
 	#include "MoviePlayer.h"
     #include "PreLoadScreenManager.h"
+	#include "InstallBundleManagerInterface.h"
 
 	#include "ShaderCodeLibrary.h"
 	#include "ShaderPipelineCache.h"
@@ -890,6 +891,11 @@ void LaunchUpdateMostRecentProjectFile()
 	}
 }
 
+#if WITH_ENGINE
+void OnStartupContentMounted(FInstallBundleResultInfo Result, bool bDumpEarlyConfigReads, bool bDumpEarlyPakFileReads);
+#endif
+void HandleConfigReload(bool bDumpEarlyConfigReads, bool bDumpEarlyPakFileReads);
+
 #if !UE_BUILD_SHIPPING
 class FFileInPakFileHistoryHelper
 {
@@ -2384,10 +2390,10 @@ int32 FEngineLoop::PreInit(const TCHAR* CmdLine)
 
 			if (bWithConfigPatching)
 			{
-				IPlatformInstallBundleManager* BundleManager = FPlatformMisc::GetPlatformInstallBundleManager();
+				IInstallBundleManager* BundleManager = IInstallBundleManager::GetPlatformInstallBundleManager();
 				if (BundleManager != nullptr && !BundleManager->IsNullInterface())
 				{
-					IPlatformInstallBundleManager::InstallBundleCompleteDelegate.AddRaw(this, &FEngineLoop::OnStartupContentMounted, bDumpEarlyConfigReads, bDumpEarlyPakFileReads);
+					IInstallBundleManager::InstallBundleCompleteDelegate.AddStatic(&OnStartupContentMounted, bDumpEarlyConfigReads, bDumpEarlyPakFileReads);
 				}
 				// If not using the bundle manager, config will be reloaded after ESP, see below
 			}
@@ -2504,7 +2510,7 @@ int32 FEngineLoop::PreInit(const TCHAR* CmdLine)
 
 		//Now that our EarlyStartupScreen is finished, lets take the necessary steps to mount paks, apply .ini cvars, and open the shader libraries if we installed content we expect to handle
 		//If using a bundle manager, assume its handling all this stuff and that we don't have to do it.
-		IPlatformInstallBundleManager* BundleManager = FPlatformMisc::GetPlatformInstallBundleManager();
+		IInstallBundleManager* BundleManager = IInstallBundleManager::GetPlatformInstallBundleManager();
 		if (BundleManager == nullptr || BundleManager->IsNullInterface())
 		{
 			// Mount Paks that were installed during EarlyStartupScreen
@@ -3623,7 +3629,7 @@ void FEngineLoop::Exit()
 	GIsRunning	= 0;
 	GLogConsole	= nullptr;
 
-	IPlatformInstallBundleManager::InstallBundleCompleteDelegate.RemoveAll(this);
+	IInstallBundleManager::InstallBundleCompleteDelegate.RemoveAll(this);
 
 	// shutdown visual logger and flush all data
 #if ENABLE_VISUAL_LOG
@@ -3780,17 +3786,19 @@ void FEngineLoop::ProcessLocalPlayerSlateOperations() const
 	}
 }
 
-void FEngineLoop::OnStartupContentMounted(FInstallBundleResultInfo Result, bool bDumpEarlyConfigReads, bool bDumpEarlyPakFileReads)
+#if WITH_ENGINE
+void OnStartupContentMounted(FInstallBundleResultInfo Result, bool bDumpEarlyConfigReads, bool bDumpEarlyPakFileReads)
 {
 	if (Result.bIsStartup && Result.Result == EInstallBundleResult::OK)
 	{
 		HandleConfigReload(bDumpEarlyConfigReads, bDumpEarlyPakFileReads);
 
-		IPlatformInstallBundleManager::InstallBundleCompleteDelegate.RemoveAll(this);
+		IInstallBundleManager::InstallBundleCompleteDelegate.RemoveAll(&GEngineLoop);
 	}
 }
+#endif
 
-void FEngineLoop::HandleConfigReload(bool bDumpEarlyConfigReads, bool bDumpEarlyPakFileReads)
+void HandleConfigReload(bool bDumpEarlyConfigReads, bool bDumpEarlyPakFileReads)
 {
 	if (bDumpEarlyConfigReads)
 	{
diff --git a/Engine/Source/Runtime/Launch/Public/LaunchEngineLoop.h b/Engine/Source/Runtime/Launch/Public/LaunchEngineLoop.h
--- a/Engine/Source/Runtime/Launch/Public/LaunchEngineLoop.h
+++ b/Engine/Source/Runtime/Launch/Public/LaunchEngineLoop.h
@@ -8,8 +8,6 @@
 	#include "UnrealEngine.h"
 #endif
 
-#include "GenericPlatform/GenericPlatformInstallBundleManager.h"
-
 class FEngineService;
 class FPendingCleanupObjects;
 class ISessionService;
@@ -118,10 +116,6 @@ private:
 	/** Utility function that processes Slate operations. */
 	void ProcessLocalPlayerSlateOperations() const;
 
-	void OnStartupContentMounted(FInstallBundleResultInfo Result, bool bDumpEarlyConfigReads, bool bDumpEarlyPakFileReads);
-
-	void HandleConfigReload(bool bDumpEarlyConfigReads, bool bDumpEarlyPakFileReads);
-
 protected:
 
 	/** Holds a dynamically expanding array of frame times in milliseconds (if FApp::IsBenchmarking() is set). */
diff --git a/Engine/Source/Runtime/NullInstallBundleManager/NullInstallBundleManager.build.cs b/Engine/Source/Runtime/NullInstallBundleManager/NullInstallBundleManager.build.cs
--- a/Engine/Source/Runtime/NullInstallBundleManager/NullInstallBundleManager.build.cs
+++ b/Engine/Source/Runtime/NullInstallBundleManager/NullInstallBundleManager.build.cs
@@ -9,6 +9,7 @@ public class NullInstallBundleManager : ModuleRules
 		PrivateDependencyModuleNames.AddRange(
 			new string[] {
 				"Core",
+				"InstallBundleManager"
 			}
 			);
 	}
diff --git a/Engine/Source/Runtime/NullInstallBundleManager/Source/NullInstallBundleManager.cpp b/Engine/Source/Runtime/NullInstallBundleManager/Source/NullInstallBundleManager.cpp
--- a/Engine/Source/Runtime/NullInstallBundleManager/Source/NullInstallBundleManager.cpp
+++ b/Engine/Source/Runtime/NullInstallBundleManager/Source/NullInstallBundleManager.cpp
@@ -1,9 +1,10 @@
 // Copyright 1998-2019 Epic Games, Inc. All Rights Reserved.
 
-#include "GenericPlatform/GenericPlatformInstallBundleManager.h"
+#include "InstallBundleManagerInterface.h"
+#include "InstallBundleManagerModule.h"
 #include "Modules/ModuleManager.h"
 
-class FNullInstallBundleManager : public IPlatformInstallBundleManager
+class FNullInstallBundleManager : public IInstallBundleManager
 {
 	virtual bool HasBuildMetaData() const override
 	{
@@ -123,7 +124,7 @@ private:
 	
 };
 
-class FNullInstallBundleManagerModule : public TPlatformInstallBundleManagerModule<FNullInstallBundleManager>
+class FNullInstallBundleManagerModule : public TInstallBundleManagerModule<FNullInstallBundleManager>
 {	
 };
 
