Non-unity build fix

diff --git a/Engine/Source/Runtime/InstallBundleManager/Public/InstallBundleManagerModule.h b/Engine/Source/Runtime/InstallBundleManager/Public/InstallBundleManagerModule.h
--- a/Engine/Source/Runtime/InstallBundleManager/Public/InstallBundleManagerModule.h
+++ b/Engine/Source/Runtime/InstallBundleManager/Public/InstallBundleManagerModule.h
@@ -3,6 +3,7 @@
 
 #include "Modules/ModuleInterface.h"
 #include "Modules/ModuleManager.h"
+#include "Misc/ConfigCacheIni.h"
 
 /**
  * Currently empty implementation for InstallBundleModule until things are moved in here.
