Install Bundle Manager
- Add TestUpdateContent() which simulates calling RequestUpdateContent() and returns which bundles are already enqueued and which bundles would be enqued.  This allows code to wait on bundles and setup progress tracking without having to be responsible for requesting them.

diff --git a/Engine/Source/Runtime/InstallBundleManager/Public/InstallBundleManagerInterface.h b/Engine/Source/Runtime/InstallBundleManager/Public/InstallBundleManagerInterface.h
--- a/Engine/Source/Runtime/InstallBundleManager/Public/InstallBundleManagerInterface.h
+++ b/Engine/Source/Runtime/InstallBundleManager/Public/InstallBundleManagerInterface.h
@@ -262,6 +262,13 @@ struct FInstallBundleRequestInfo
 	TArray<FName> BundlesQueuedForRemoval;
 };
 
+struct FInstallBundleTestInfo
+{
+	EInstallBundleRequestInfoFlags InfoFlags = EInstallBundleRequestInfoFlags::None;
+	TArray<FName> BundlesQueuedForInstall;
+	TArray<FName> BundlesNeededForInstall;
+};
+
 enum class EInstallBundleCancelFlags : int32
 {
 	None		= 0,
@@ -304,6 +311,9 @@ public:
 
 	virtual bool IsActive() const = 0;
 
+	virtual FInstallBundleTestInfo TestUpdateContent(FName BundleName) = 0;
+	virtual FInstallBundleTestInfo TestUpdateContent(TArrayView<FName> BundleNames) = 0;
+
 	virtual FInstallBundleRequestInfo RequestUpdateContent(FName BundleName, EInstallBundleRequestFlags Flags) = 0;
 	virtual FInstallBundleRequestInfo RequestUpdateContent(TArrayView<FName> BundleNames, EInstallBundleRequestFlags Flags) = 0;
 
diff --git a/Engine/Source/Runtime/NullInstallBundleManager/Source/NullInstallBundleManager.cpp b/Engine/Source/Runtime/NullInstallBundleManager/Source/NullInstallBundleManager.cpp
--- a/Engine/Source/Runtime/NullInstallBundleManager/Source/NullInstallBundleManager.cpp
+++ b/Engine/Source/Runtime/NullInstallBundleManager/Source/NullInstallBundleManager.cpp
@@ -33,12 +33,20 @@ class FNullInstallBundleManager : public IInstallBundleManager
 		return false;
 	}
 
+	virtual FInstallBundleTestInfo TestUpdateContent(FName BundleName) override
+	{
+		return FInstallBundleTestInfo();
+	}
+	virtual FInstallBundleTestInfo TestUpdateContent(TArrayView<FName> BundleNames) override
+	{
+		return FInstallBundleTestInfo();
+	}
+
 	virtual FInstallBundleRequestInfo RequestUpdateContent(FName BundleName, EInstallBundleRequestFlags Flags) override
 	{
 		FInstallBundleRequestInfo RetInfo;
 		return RetInfo;
 	}
-
 	virtual FInstallBundleRequestInfo RequestUpdateContent(TArrayView<FName> BundleNames, EInstallBundleRequestFlags Flags) override
 	{
 		FInstallBundleRequestInfo RetInfo;
