Add CombinedPauseFlags to FBundlePrereqCombinedStatusHelper::FCombinedBundleStatus

diff --git a/Engine/Source/Runtime/InstallBundleManager/Private/BundlePrereqCombinedStatusHelper.cpp b/Engine/Source/Runtime/InstallBundleManager/Private/BundlePrereqCombinedStatusHelper.cpp
--- a/Engine/Source/Runtime/InstallBundleManager/Private/BundlePrereqCombinedStatusHelper.cpp
+++ b/Engine/Source/Runtime/InstallBundleManager/Private/BundlePrereqCombinedStatusHelper.cpp
@@ -4,6 +4,8 @@
 #include "Containers/Ticker.h"
 #include "InstallBundleManagerPrivatePCH.h"
 
+//PRAGMA_DISABLE_OPTIMIZATION
+
 FBundlePrereqCombinedStatusHelper::FBundlePrereqCombinedStatusHelper()
 : DownloadWeight(1.f)
 , InstallWeight(1.f)
@@ -152,6 +154,7 @@ void FBundlePrereqCombinedStatusHelper::UpdateCombinedStatus()
 	CurrentCombinedStatus.ProgressPercent = GetCombinedProgressPercent();
 	
 	EInstallBundleStatus EarliestBundleState = EInstallBundleStatus::Count;
+	EInstallBundlePauseFlags CombinedPauseFlags = EInstallBundlePauseFlags::None;
 	bool bIsAnythingPaused = false;
 	bool bIsAnythingFinishing = false;
 	
@@ -180,11 +183,20 @@ void FBundlePrereqCombinedStatusHelper::UpdateCombinedStatus()
 		}
 		
 		bIsAnythingPaused = bIsAnythingPaused || BundlePair.Value.PauseFlags != EInstallBundlePauseFlags::None;
+		CombinedPauseFlags |= BundlePair.Value.PauseFlags;
 	}
 	
 	//if we have any paused bundles, and we have any bundle that isn't finished installed, we are Paused
 	//if everything is installed ignore the pause flags as we completed after pausing the bundles
 	CurrentCombinedStatus.bIsPaused = (bIsAnythingPaused && (EarliestBundleState < EInstallBundleStatus::Installed));
+	if(CurrentCombinedStatus.bIsPaused)
+	{
+		CurrentCombinedStatus.CombinedPauseFlags = CombinedPauseFlags;
+	}
+	else
+	{
+		CurrentCombinedStatus.CombinedPauseFlags = EInstallBundlePauseFlags::None;
+	}
 	
 	//if the bundle does not need an update, all the phases we go through don't support pausing (Mounting ,Compiling Shaders, etc)
 	//Otherwise start with True and override those specific cases bellow
@@ -327,3 +339,5 @@ void FBundlePrereqCombinedStatusHelper::OnBundleInstallComplete(FInstallBundleRe
 		}
 	}
 }
+
+//PRAGMA_ENABLE_OPTIMIZATION
diff --git a/Engine/Source/Runtime/InstallBundleManager/Public/BundlePrereqCombinedStatusHelper.h b/Engine/Source/Runtime/InstallBundleManager/Public/BundlePrereqCombinedStatusHelper.h
--- a/Engine/Source/Runtime/InstallBundleManager/Public/BundlePrereqCombinedStatusHelper.h
+++ b/Engine/Source/Runtime/InstallBundleManager/Public/BundlePrereqCombinedStatusHelper.h
@@ -23,20 +23,12 @@ public:
 			,Finished
 		};
 		
-		float ProgressPercent;
-		ECombinedBundleStateEnum CombinedState;
-		bool bIsPaused;
-		bool bDoesCurrentStateSupportPausing;
-		bool bBundleRequiresUpdate;
-		
-		FCombinedBundleStatus()
-		: ProgressPercent(0.f)
-		, CombinedState(ECombinedBundleStateEnum::Unknown)
-		, bIsPaused(false)
-		, bDoesCurrentStateSupportPausing(false)
-		, bBundleRequiresUpdate(false)
-		{
-		}
+		float ProgressPercent = 0.0f;
+		ECombinedBundleStateEnum CombinedState = ECombinedBundleStateEnum::Unknown;
+		EInstallBundlePauseFlags CombinedPauseFlags = EInstallBundlePauseFlags::None;
+		bool bIsPaused = false;
+		bool bDoesCurrentStateSupportPausing = false;
+		bool bBundleRequiresUpdate = false;
 	};
 	
 public:
