InstallBundleManager - Uninstalling now correctly handles dependent bundles.

diff --git a/Engine/Source/Runtime/InstallBundleManager/Public/InstallBundleManagerInterface.h b/Engine/Source/Runtime/InstallBundleManager/Public/InstallBundleManagerInterface.h
--- a/Engine/Source/Runtime/InstallBundleManager/Public/InstallBundleManagerInterface.h
+++ b/Engine/Source/Runtime/InstallBundleManager/Public/InstallBundleManagerInterface.h
@@ -324,8 +324,8 @@ public:
     
 	virtual FInstallBundleRequestInfo RequestRemoveContent(FName BundleName) = 0;
 
-	virtual void RequestRemoveContentOnNextInit(FName BundleName) = 0;
-	virtual void RequestRemoveContentOnNextInit(TArrayView<FName> BundleNames) = 0;
+	virtual void RequestRemoveContentOnNextInit(FName RemoveName, TArrayView<FName> KeepNames = TArrayView<FName>()) = 0;
+	virtual void RequestRemoveContentOnNextInit(TArrayView<FName> RemoveNames, TArrayView<FName> KeepNames = TArrayView<FName>()) = 0;
 
 	virtual void CancelRequestRemoveContentOnNextInit(FName BundleName) = 0;
 	virtual void CancelRequestRemoveContentOnNextInit(TArrayView<FName> BundleName) = 0;
diff --git a/Engine/Source/Runtime/NullInstallBundleManager/Source/NullInstallBundleManager.cpp b/Engine/Source/Runtime/NullInstallBundleManager/Source/NullInstallBundleManager.cpp
--- a/Engine/Source/Runtime/NullInstallBundleManager/Source/NullInstallBundleManager.cpp
+++ b/Engine/Source/Runtime/NullInstallBundleManager/Source/NullInstallBundleManager.cpp
@@ -77,10 +77,10 @@ class FNullInstallBundleManager : public IInstallBundleManager
 	{
 	}
 
-	virtual void RequestRemoveContentOnNextInit(FName BundleName) override
+	virtual void RequestRemoveContentOnNextInit(FName RemoveName, TArrayView<FName> KeepNames = TArrayView<FName>()) override
 	{
 	}
-	virtual void RequestRemoveContentOnNextInit(TArrayView<FName> BundleNames) override
+	virtual void RequestRemoveContentOnNextInit(TArrayView<FName> RemoveNames, TArrayView<FName> KeepNames = TArrayView<FName>()) override
 	{
 	}
 
