Install Bundle Manager Refactor to allow adding arbitray bundle sources
- Created Install Bundle Types as a place to put shared types
- Created Install Bundle Utils as a place to put shared utilities
- Removed unused API functionality - TestUpdateContent, RequestRemoveContent, IsActive
- Removed unused debug functionality
- Added InstallBundleSource Interface and made a bulk bundle source and a BPS bundle source
- Created a "NEW" bundle manager that can use the sources, this will allow me to do incremental checkins for the refactor without chaning the existing implementation.
- Move a bunch of helper types that now need slightly different implentations into the namespace of their respective classes  so they can co-exist
- "NEW" bundle manager can be fully initialized with any combination of bundle sources
- Bulk and BPS bundle sources can be fully initialized

diff --git a/Engine/Source/Runtime/InstallBundleManager/Private/InstallBundleManagerPrivatePCH.h b/Engine/Source/Runtime/InstallBundleManager/Private/InstallBundleManagerPrivatePCH.h
--- a/Engine/Source/Runtime/InstallBundleManager/Private/InstallBundleManagerPrivatePCH.h
+++ b/Engine/Source/Runtime/InstallBundleManager/Private/InstallBundleManagerPrivatePCH.h
@@ -3,6 +3,5 @@
 #pragma once
 
 #include "CoreMinimal.h"
-#include "InstallBundleManagerInterface.h"
 
 DECLARE_LOG_CATEGORY_EXTERN(LogInstallBundleManager, Log, All);
diff --git a/Engine/Source/Runtime/InstallBundleManager/Private/InstallBundleTypes.cpp b/Engine/Source/Runtime/InstallBundleManager/Private/InstallBundleTypes.cpp
new file mode 100644
--- /dev/null
+++ b/Engine/Source/Runtime/InstallBundleManager/Private/InstallBundleTypes.cpp
@@ -0,0 +1,65 @@
+// Copyright 1998-2019 Epic Games, Inc. All Rights Reserved.
+
+#include "InstallBundleTypes.h"
+#include "InstallBundleManagerPrivatePCH.h"
+
+#include "InstallBundleUtils.h"
+
+const TCHAR* LexToString(EInstallBundleSourceType Type)
+{
+	static const TCHAR* Strings[] =
+	{
+		TEXT("Bulk"),
+		TEXT("BuildPatchServices"),
+	};
+	static_assert(InstallBundleUtil::CastToUnderlying(EInstallBundleSourceType::Count) == UE_ARRAY_COUNT(Strings), "");
+	return Strings[InstallBundleUtil::CastToUnderlying(Type)];
+}
+
+const TCHAR* LexToString(EInstallBundleManagerInitResult Result)
+{
+	static const TCHAR* Strings[] =
+	{
+		TEXT("OK"),
+		TEXT("BuildMetaDataNotFound"),
+		TEXT("BuildMetaDataDownloadError"),
+		TEXT("BuildMetaDataParsingError"),
+		TEXT("DistributionRootParseError"),
+		TEXT("DistributionRootDownloadError"),
+		TEXT("ManifestArchiveError"),
+		TEXT("ManifestCreationError"),
+		TEXT("ManifestDownloadError"),
+		TEXT("BackgroundDownloadsIniDownloadError"),
+		TEXT("NoInternetConnectionError"),
+		TEXT("ConfigurationError"),
+	};
+	static_assert(InstallBundleUtil::CastToUnderlying(EInstallBundleManagerInitResult::Count) == UE_ARRAY_COUNT(Strings), "");
+	return Strings[InstallBundleUtil::CastToUnderlying(Result)];
+}
+
+const TCHAR* LexToString(EBundleState Val)
+{
+	static const TCHAR* Strings[] =
+	{
+		TEXT("NotInstalled"),
+		TEXT("NeedsUpdate"),
+		TEXT("NeedsMount"),
+		TEXT("Mounted"),
+	};
+	static_assert(InstallBundleUtil::CastToUnderlying(EBundleState::Count) == UE_ARRAY_COUNT(Strings), "");
+	return Strings[InstallBundleUtil::CastToUnderlying(Val)];
+}
+
+const TCHAR* LexToString(EInstallBundleContentState State)
+{
+	static const TCHAR* Strings[] =
+	{
+		TEXT("InitializationError"),
+		TEXT("NotInstalled"),
+		TEXT("NeedsUpdate"),
+		TEXT("UpToDate"),
+	};
+	static_assert(InstallBundleUtil::CastToUnderlying(EInstallBundleContentState::Count) == UE_ARRAY_COUNT(Strings), "");
+
+	return Strings[InstallBundleUtil::CastToUnderlying(State)];
+}
diff --git a/Engine/Source/Runtime/InstallBundleManager/Private/InstallBundleUtils.cpp b/Engine/Source/Runtime/InstallBundleManager/Private/InstallBundleUtils.cpp
new file mode 100644
--- /dev/null
+++ b/Engine/Source/Runtime/InstallBundleManager/Private/InstallBundleUtils.cpp
@@ -0,0 +1,58 @@
+// Copyright 1998-2019 Epic Games, Inc. All Rights Reserved.
+
+#include "InstallBundleUtils.h"
+#include "InstallBundleManagerPrivatePCH.h"
+
+namespace InstallBundleUtil
+{
+	bool HasInternetConnection(ENetworkConnectionType ConnectionType)
+	{
+		return ConnectionType != ENetworkConnectionType::AirplaneMode
+			&& ConnectionType != ENetworkConnectionType::None;
+	}
+
+	FName FInstallBundleManagerKeepAwake::Tag(TEXT("InstallBundleManagerKeepAwake"));
+	FName FInstallBundleManagerKeepAwake::TagWithRendering(TEXT("InstallBundleManagerKeepAwakeWithRendering"));
+
+	void StartInstallBundleAsyncIOTask(TArray<TUniquePtr<FInstallBundleTask>>& Tasks, TUniqueFunction<void()> WorkFunc, TUniqueFunction<void()> OnComplete)
+	{
+		TUniquePtr<FInstallBundleTask> Task = MakeUnique<FInstallBundleTask>(MoveTemp(WorkFunc), MoveTemp(OnComplete));
+		Task->StartBackgroundTask(GIOThreadPool);
+		Tasks.Add(MoveTemp(Task));
+	}
+
+	void FinishInstallBundleAsyncIOTasks(TArray<TUniquePtr<FInstallBundleTask>>& Tasks)
+	{
+		TArray<TUniquePtr<FInstallBundleTask>> FinishedTasks;
+		for (int32 i = 0; i < Tasks.Num();)
+		{
+			TUniquePtr<FInstallBundleTask>& Task = Tasks[i];
+			check(Task);
+			if (Task->IsDone())
+			{
+				FinishedTasks.Add(MoveTemp(Task));
+				Tasks.RemoveAtSwap(i, 1, false);
+			}
+			else
+			{
+				++i;
+			}
+		}
+		for (TUniquePtr<FInstallBundleTask>& Task : FinishedTasks)
+		{
+			Task->GetTask().CallOnComplete();
+		}
+	}
+
+	void CleanupInstallBundleAsyncIOTasks(TArray<TUniquePtr<FInstallBundleTask>>& Tasks)
+	{
+		for (TUniquePtr<FInstallBundleTask>& Task : Tasks)
+		{
+			check(Task);
+			if (!Task->Cancel())
+			{
+				Task->EnsureCompletion(false);
+			}
+		}
+	}
+}
diff --git a/Engine/Source/Runtime/InstallBundleManager/Public/InstallBundleManagerInterface.h b/Engine/Source/Runtime/InstallBundleManager/Public/InstallBundleManagerInterface.h
--- a/Engine/Source/Runtime/InstallBundleManager/Public/InstallBundleManagerInterface.h
+++ b/Engine/Source/Runtime/InstallBundleManager/Public/InstallBundleManagerInterface.h
@@ -2,6 +2,7 @@
 
 #pragma once
 
+#include "InstallBundleTypes.h"
 #include "Misc/EnumClassFlags.h"
 #include "Misc/ConfigCacheIni.h"
 #include "Logging/LogMacros.h"
@@ -9,46 +10,6 @@
 
 class IAnalyticsProviderET;
 
-enum class EInstallBundleModuleInitResult : int
-{
-	OK,
-	BuildMetaDataNotFound,
-	BuildMetaDataDownloadError,
-	BuildMetaDataParsingError,
-	DistributionRootParseError,
-	DistributionRootDownloadError,
-	ManifestArchiveError,
-	ManifestCreationError,
-	ManifestDownloadError,
-	BackgroundDownloadsIniDownloadError,
-	NoInternetConnectionError,
-	ConfigurationError,
-	Count
-};
-
-inline const TCHAR* LexToString(EInstallBundleModuleInitResult Result)
-{
-	using UnderType = __underlying_type(EInstallBundleModuleInitResult);
-	static const TCHAR* Strings[] =
-	{
-		TEXT("OK"),
-		TEXT("BuildMetaDataNotFound"),
-		TEXT("BuildMetaDataDownloadError"),
-		TEXT("BuildMetaDataParsingError"),
-		TEXT("DistributionRootParseError"),
-		TEXT("DistributionRootDownloadError"),
-		TEXT("ManifestArchiveError"),
-		TEXT("ManifestCreationError"),
-		TEXT("ManifestDownloadError"),
-		TEXT("BackgroundDownloadsIniDownloadError"),
-		TEXT("NoInternetConnectionError"),
-		TEXT("ConfigurationError"),
-	};
-	static_assert(static_cast<UnderType>(EInstallBundleModuleInitResult::Count) == UE_ARRAY_COUNT(Strings), "");
-
-	return Strings[static_cast<UnderType>(Result)];
-}
-
 enum class EInstallBundleResult : int
 {
 	OK,
@@ -203,57 +164,15 @@ struct FInstallBundlePauseInfo
 	EInstallBundlePauseFlags PauseFlags = EInstallBundlePauseFlags::None;
 };
 
-enum class EInstallBundleContentState : int
-{
-	InitializationError,
-	NotInstalled,
-	NeedsUpdate,
-	UpToDate,
-	Count,
-};
-inline const TCHAR* LexToString(EInstallBundleContentState State)
-{
-	using UnderType = __underlying_type(EInstallBundleContentState);
-	static const TCHAR* Strings[] =
-	{
-		TEXT("InitializationError"),
-		TEXT("NotInstalled"),
-		TEXT("NeedsUpdate"),
-		TEXT("UpToDate"),
-	};
-	static_assert(static_cast<UnderType>(EInstallBundleContentState::Count) == UE_ARRAY_COUNT(Strings), "");
-
-	return Strings[static_cast<UnderType>(State)];
-}
-
-struct FInstallBundleContentState
-{
-	EInstallBundleContentState State = EInstallBundleContentState::InitializationError;
-	TMap<FName, EInstallBundleContentState> IndividualBundleStates;
-	TMap<FName, float> IndividualBundleWeights;
-	uint64 DownloadSize = 0;
-	uint64 InstallSize = 0;
-	uint64 InstallOverheadSize = 0;
-	uint64 FreeSpace = 0;
-};
-
-enum class EInstallBundleGetContentStateFlags : uint32
-{
-	None = 0,
-	ForceNoPatching = (1 << 0),
-};
-ENUM_CLASS_FLAGS(EInstallBundleGetContentStateFlags);
-
 enum class EInstallBundleRequestInfoFlags : int32
 {
 	None							= 0,
 	EnqueuedBundlesForInstall		= (1 << 0),
 	EnqueuedBundlesForRemoval		= (1 << 1),
 	SkippedAlreadyMountedBundles	= (1 << 2),
-	SkippedBundlesQueuedForRemoval	= (1 << 3),
-	SkippedBundlesQueuedForInstall  = (1 << 4), // Only valid for removal requests
-	SkippedUnknownBundles			= (1 << 5),
-	InitializationError				= (1 << 6), // Can't enqueue because the bundle manager failed to initialize
+	SkippedBundlesQueuedForInstall  = (1 << 3), // Only valid for removal requests
+	SkippedUnknownBundles			= (1 << 4),
+	InitializationError				= (1 << 5), // Can't enqueue because the bundle manager failed to initialize
 };
 ENUM_CLASS_FLAGS(EInstallBundleRequestInfoFlags);
 
@@ -261,14 +180,6 @@ struct FInstallBundleRequestInfo
 {
 	EInstallBundleRequestInfoFlags InfoFlags = EInstallBundleRequestInfoFlags::None;
 	TArray<FName> BundlesQueuedForInstall;
-	TArray<FName> BundlesQueuedForRemoval;
-};
-
-struct FInstallBundleTestInfo
-{
-	EInstallBundleRequestInfoFlags InfoFlags = EInstallBundleRequestInfoFlags::None;
-	TArray<FName> BundlesQueuedForInstall;
-	TArray<FName> BundlesNeededForInstall;
 };
 
 enum class EInstallBundleCancelFlags : int32
@@ -285,7 +196,7 @@ enum class EInstallBundleManagerInitErrorHandlerResult
 	StopInitialization, // Stop trying to initialize
 };
 
-DECLARE_DELEGATE_RetVal_OneParam(EInstallBundleManagerInitErrorHandlerResult, FInstallBundleManagerInitErrorHandler, EInstallBundleModuleInitResult);
+DECLARE_DELEGATE_RetVal_OneParam(EInstallBundleManagerInitErrorHandlerResult, FInstallBundleManagerInitErrorHandler, EInstallBundleManagerInitResult);
 
 DECLARE_MULTICAST_DELEGATE_OneParam(FInstallBundleCompleteMultiDelegate, FInstallBundleResultInfo);
 DECLARE_MULTICAST_DELEGATE_OneParam(FInstallBundlePausedMultiDelegate, FInstallBundlePauseInfo);
@@ -303,6 +214,7 @@ public:
 
 	virtual ~IInstallBundleManager() {}
 
+	// TODO: should query if we are using the BPT source
 	virtual bool HasBuildMetaData() const = 0;
 
 	virtual FDelegateHandle PushInitErrorCallback(FInstallBundleManagerInitErrorHandler Callback) = 0;
@@ -310,13 +222,7 @@ public:
 	virtual void PopInitErrorCallback(FDelegateHandle Handle) = 0;
 	virtual void PopInitErrorCallback(const void* InUserObject) = 0;
 
-	virtual bool IsInitialized() const = 0;
-	virtual bool IsInitializing() const = 0;
-
-	virtual bool IsActive() const = 0;
-
-	virtual FInstallBundleTestInfo TestUpdateContent(FName BundleName) = 0;
-	virtual FInstallBundleTestInfo TestUpdateContent(TArrayView<FName> BundleNames) = 0;
+	virtual EInstallBundleManagerInitState GetInitState() const = 0;
 
 	virtual FInstallBundleRequestInfo RequestUpdateContent(FName BundleName, EInstallBundleRequestFlags Flags) = 0;
 	virtual FInstallBundleRequestInfo RequestUpdateContent(TArrayView<FName> BundleNames, EInstallBundleRequestFlags Flags) = 0;
@@ -326,8 +232,6 @@ public:
 
     virtual void CancelAllGetContentStateRequestsForTag(FName RequestTag) = 0;
     
-	virtual FInstallBundleRequestInfo RequestRemoveContent(FName BundleName) = 0;
-
 	virtual void RequestRemoveContentOnNextInit(FName RemoveName, TArrayView<FName> KeepNames = TArrayView<FName>()) = 0;
 	virtual void RequestRemoveContentOnNextInit(TArrayView<FName> RemoveNames, TArrayView<FName> KeepNames = TArrayView<FName>()) = 0;
 
diff --git a/Engine/Source/Runtime/InstallBundleManager/Public/InstallBundleManagerModule.h b/Engine/Source/Runtime/InstallBundleManager/Public/InstallBundleManagerModule.h
--- a/Engine/Source/Runtime/InstallBundleManager/Public/InstallBundleManagerModule.h
+++ b/Engine/Source/Runtime/InstallBundleManager/Public/InstallBundleManagerModule.h
@@ -4,6 +4,7 @@
 #include "Modules/ModuleInterface.h"
 #include "Modules/ModuleManager.h"
 #include "Misc/ConfigCacheIni.h"
+#include "InstallBundleManagerInterface.h"
 
 /**
  * Currently empty implementation for InstallBundleModule until things are moved in here.
diff --git a/Engine/Source/Runtime/InstallBundleManager/Public/InstallBundleSourceInterface.h b/Engine/Source/Runtime/InstallBundleManager/Public/InstallBundleSourceInterface.h
new file mode 100644
--- /dev/null
+++ b/Engine/Source/Runtime/InstallBundleManager/Public/InstallBundleSourceInterface.h
@@ -0,0 +1,29 @@
+// Copyright 1998-2019 Epic Games, Inc. All Rights Reserved.
+
+#pragma once
+
+#include "InstallBundleTypes.h"
+
+class IInstallBundleSource;
+class IAnalyticsProviderET;
+
+DECLARE_DELEGATE_ThreeParams(FInstallBundleSourceInitDelegate, TSharedRef<IInstallBundleSource> /*Source*/, EInstallBundleManagerInitResult /*Result*/, bool /*bShouldUseFallbackSource*/)
+
+class INSTALLBUNDLEMANAGER_API IInstallBundleSource : public TSharedFromThis<IInstallBundleSource>
+{
+public:
+	virtual ~IInstallBundleSource() {}
+
+	virtual EInstallBundleSourceType GetSourceType() const = 0;
+
+	virtual void Init(
+		TSharedRef<FConfigFile> InstallBundleConfig, 
+		TSharedPtr<IAnalyticsProviderET> AnalyticsProvider) = 0;
+	virtual void AsyncInit(FInstallBundleSourceInitDelegate Callback) = 0;
+
+	virtual EInstallBundleManagerInitState GetInitState() const = 0;
+	// Only valid after AsyncInit completes
+	virtual EInstallBundleManagerInitResult GetLastInitResult() const = 0;
+
+	virtual void SetErrorSimulationCommands(const FString& CommandLine) {}
+};
diff --git a/Engine/Source/Runtime/InstallBundleManager/Public/InstallBundleTypes.h b/Engine/Source/Runtime/InstallBundleManager/Public/InstallBundleTypes.h
new file mode 100644
--- /dev/null
+++ b/Engine/Source/Runtime/InstallBundleManager/Public/InstallBundleTypes.h
@@ -0,0 +1,80 @@
+// Copyright 1998-2019 Epic Games, Inc. All Rights Reserved.
+
+#pragma once
+
+#include "CoreMinimal.h"
+#include "Misc/EnumRange.h"
+
+enum class EInstallBundleSourceType : int
+{
+	Bulk,
+	BuildPatchServices,
+	Count
+};
+ENUM_RANGE_BY_COUNT(EInstallBundleSourceType, EInstallBundleSourceType::Count);
+const TCHAR* LexToString(EInstallBundleSourceType Type);
+
+enum class EInstallBundleManagerInitState : int
+{
+	NotInitialized,
+	Failed,
+	Succeeded
+};
+
+enum class EInstallBundleManagerInitResult : int
+{
+	OK,
+	BuildMetaDataNotFound,
+	BuildMetaDataDownloadError,
+	BuildMetaDataParsingError,
+	DistributionRootParseError,
+	DistributionRootDownloadError,
+	ManifestArchiveError,
+	ManifestCreationError,
+	ManifestDownloadError,
+	BackgroundDownloadsIniDownloadError,
+	NoInternetConnectionError,
+	ConfigurationError,
+	Count
+};
+const TCHAR* LexToString(EInstallBundleManagerInitResult Result);
+
+// TODO: Needs to be renamed to EInstallBundleState
+enum class EBundleState : int
+{
+	NotInstalled,
+	NeedsUpdate,
+	NeedsMount,
+	Mounted,
+	Count,
+};
+const TCHAR* LexToString(EBundleState Val);
+
+enum class EInstallBundleContentState : int
+{
+	InitializationError,
+	NotInstalled,
+	NeedsUpdate,
+	UpToDate,
+	Count,
+};
+const TCHAR* LexToString(EInstallBundleContentState State);
+
+struct FInstallBundleContentState
+{
+	EInstallBundleContentState State = EInstallBundleContentState::InitializationError;
+	TMap<FName, EInstallBundleContentState> IndividualBundleStates;
+	TMap<FName, float> IndividualBundleWeights;
+	uint64 DownloadSize = 0;
+	uint64 InstallSize = 0;
+	uint64 InstallOverheadSize = 0;
+	uint64 FreeSpace = 0;
+};
+
+enum class EInstallBundleGetContentStateFlags : uint32
+{
+	None = 0,
+	ForceNoPatching = (1 << 0),
+};
+ENUM_CLASS_FLAGS(EInstallBundleGetContentStateFlags);
+
diff --git a/Engine/Source/Runtime/InstallBundleManager/Public/InstallBundleUtils.h b/Engine/Source/Runtime/InstallBundleManager/Public/InstallBundleUtils.h
new file mode 100644
--- /dev/null
+++ b/Engine/Source/Runtime/InstallBundleManager/Public/InstallBundleUtils.h
@@ -0,0 +1,91 @@
+// Copyright 1998-2019 Epic Games, Inc. All Rights Reserved.
+
+#pragma once
+
+#include "CoreMinimal.h"
+#include "Async/AsyncWork.h"
+#include "Misc/EmbeddedCommunication.h"
+
+namespace InstallBundleUtil
+{
+	bool HasInternetConnection(ENetworkConnectionType ConnectionType);
+
+	// It would really be nice to have these in core
+	template<class EnumType>
+	constexpr auto& CastAsUnderlying(EnumType &Type)
+	{
+		static_assert(TIsEnum<EnumType>::Value, "");
+		using UnderType = __underlying_type(EnumType);
+		return *reinterpret_cast<UnderType*>(&Type);
+	}
+
+	template<class EnumType>
+	constexpr const auto& CastAsUnderlying(const EnumType &Type)
+	{
+		static_assert(TIsEnum<EnumType>::Value, "");
+		using UnderType = __underlying_type(EnumType);
+		return *reinterpret_cast<const UnderType*>(&Type);
+	}
+
+	template<class EnumType>
+	constexpr auto CastToUnderlying(EnumType Type)
+	{
+		static_assert(TIsEnum<EnumType>::Value, "");
+		using UnderType = __underlying_type(EnumType);
+		return static_cast<UnderType>(Type);
+	}
+
+	// Keep the engine awake via RAII when running as an embedded app
+	class FInstallBundleManagerKeepAwake : public FEmbeddedKeepAwake
+	{
+		static FName Tag;
+		static FName TagWithRendering;
+	public:
+		FInstallBundleManagerKeepAwake(bool bNeedsRendering = false)
+			: FEmbeddedKeepAwake(bNeedsRendering ? TagWithRendering : Tag, bNeedsRendering) {}
+	};
+
+	class FInstallBundleWork : public FNonAbandonableTask
+	{
+	public:
+		FInstallBundleWork() = default;
+
+		FInstallBundleWork(TUniqueFunction<void()> InWork, TUniqueFunction<void()> InOnComplete)
+			: WorkFunc(MoveTemp(InWork))
+			, OnCompleteFunc(MoveTemp(InOnComplete))
+		{}
+
+		void DoWork()
+		{
+			if (WorkFunc)
+			{
+				WorkFunc();
+			}
+		}
+
+		void CallOnComplete()
+		{
+			if (OnCompleteFunc)
+			{
+				OnCompleteFunc();
+			}
+		}
+
+		FORCEINLINE TStatId GetStatId() const
+		{
+			RETURN_QUICK_DECLARE_CYCLE_STAT(FInstallBundleWork, STATGROUP_ThreadPoolAsyncTasks);
+		}
+
+	private:
+		TUniqueFunction<void()> WorkFunc;
+		TUniqueFunction<void()> OnCompleteFunc;
+	};
+
+	using FInstallBundleTask = FAsyncTask<FInstallBundleWork>;
+
+	void StartInstallBundleAsyncIOTask(TArray<TUniquePtr<FInstallBundleTask>>& Tasks, TUniqueFunction<void()> WorkFunc, TUniqueFunction<void()> OnComplete);
+
+	void FinishInstallBundleAsyncIOTasks(TArray<TUniquePtr<FInstallBundleTask>>& Tasks);
+
+	void CleanupInstallBundleAsyncIOTasks(TArray<TUniquePtr<FInstallBundleTask>>& Tasks);
+}
diff --git a/Engine/Source/Runtime/NullInstallBundleManager/Source/NullInstallBundleManager.cpp b/Engine/Source/Runtime/NullInstallBundleManager/Source/NullInstallBundleManager.cpp
--- a/Engine/Source/Runtime/NullInstallBundleManager/Source/NullInstallBundleManager.cpp
+++ b/Engine/Source/Runtime/NullInstallBundleManager/Source/NullInstallBundleManager.cpp
@@ -28,27 +28,9 @@ class FNullInstallBundleManager : public IInstallBundleManager
 	{
 	}
 
-	virtual bool IsInitializing() const override
+	virtual EInstallBundleManagerInitState GetInitState() const override
 	{
-		return false;
-	}
-	virtual bool IsInitialized() const override
-	{
-		return true;
-	}
-
-	virtual bool IsActive() const override
-	{
-		return false;
-	}
-
-	virtual FInstallBundleTestInfo TestUpdateContent(FName BundleName) override
-	{
-		return FInstallBundleTestInfo();
-	}
-	virtual FInstallBundleTestInfo TestUpdateContent(TArrayView<FName> BundleNames) override
-	{
-		return FInstallBundleTestInfo();
+		return EInstallBundleManagerInitState::Succeeded;
 	}
 
 	virtual FInstallBundleRequestInfo RequestUpdateContent(FName BundleName, EInstallBundleRequestFlags Flags) override
@@ -62,12 +44,6 @@ class FNullInstallBundleManager : public IInstallBundleManager
 		return RetInfo;
 	}
 
-	virtual FInstallBundleRequestInfo RequestRemoveContent(FName BundleName) override
-	{
-		FInstallBundleRequestInfo RetInfo;
-		return RetInfo;
-	}
-
 	virtual void GetContentState(FName BundleName, EInstallBundleGetContentStateFlags Flags, bool bAddDependencies, FInstallBundleGetContentStateDelegate Callback, FName RequestTag) override
 	{
 		FInstallBundleContentState State;
