Fix unresolved symbols.

diff --git a/Engine/Source/Runtime/InstallBundleManager/Public/InstallBundleSourceInterface.h b/Engine/Source/Runtime/InstallBundleManager/Public/InstallBundleSourceInterface.h
--- a/Engine/Source/Runtime/InstallBundleManager/Public/InstallBundleSourceInterface.h
+++ b/Engine/Source/Runtime/InstallBundleManager/Public/InstallBundleSourceInterface.h
@@ -9,7 +9,7 @@ class IAnalyticsProviderET;
 
 DECLARE_DELEGATE_ThreeParams(FInstallBundleSourceInitDelegate, TSharedRef<IInstallBundleSource> /*Source*/, EInstallBundleManagerInitResult /*Result*/, bool /*bShouldUseFallbackSource*/)
 
-class INSTALLBUNDLEMANAGER_API IInstallBundleSource : public TSharedFromThis<IInstallBundleSource>
+class IInstallBundleSource : public TSharedFromThis<IInstallBundleSource>
 {
 public:
 	virtual ~IInstallBundleSource() {}
diff --git a/Engine/Source/Runtime/InstallBundleManager/Public/InstallBundleTypes.h b/Engine/Source/Runtime/InstallBundleManager/Public/InstallBundleTypes.h
--- a/Engine/Source/Runtime/InstallBundleManager/Public/InstallBundleTypes.h
+++ b/Engine/Source/Runtime/InstallBundleManager/Public/InstallBundleTypes.h
@@ -12,7 +12,7 @@ enum class EInstallBundleSourceType : int
 	Count
 };
 ENUM_RANGE_BY_COUNT(EInstallBundleSourceType, EInstallBundleSourceType::Count);
-const TCHAR* LexToString(EInstallBundleSourceType Type);
+INSTALLBUNDLEMANAGER_API const TCHAR* LexToString(EInstallBundleSourceType Type);
 
 enum class EInstallBundleManagerInitState : int
 {
@@ -37,7 +37,7 @@ enum class EInstallBundleManagerInitResult : int
 	ConfigurationError,
 	Count
 };
-const TCHAR* LexToString(EInstallBundleManagerInitResult Result);
+INSTALLBUNDLEMANAGER_API const TCHAR* LexToString(EInstallBundleManagerInitResult Result);
 
 // TODO: Needs to be renamed to EInstallBundleState
 enum class EBundleState : int
@@ -48,7 +48,7 @@ enum class EBundleState : int
 	Mounted,
 	Count,
 };
-const TCHAR* LexToString(EBundleState Val);
+INSTALLBUNDLEMANAGER_API const TCHAR* LexToString(EBundleState Val);
 
 enum class EInstallBundleContentState : int
 {
@@ -58,9 +58,9 @@ enum class EInstallBundleContentState : int
 	UpToDate,
 	Count,
 };
-const TCHAR* LexToString(EInstallBundleContentState State);
+INSTALLBUNDLEMANAGER_API const TCHAR* LexToString(EInstallBundleContentState State);
 
-struct FInstallBundleContentState
+struct INSTALLBUNDLEMANAGER_API FInstallBundleContentState
 {
 	EInstallBundleContentState State = EInstallBundleContentState::InitializationError;
 	TMap<FName, EInstallBundleContentState> IndividualBundleStates;
diff --git a/Engine/Source/Runtime/InstallBundleManager/Public/InstallBundleUtils.h b/Engine/Source/Runtime/InstallBundleManager/Public/InstallBundleUtils.h
--- a/Engine/Source/Runtime/InstallBundleManager/Public/InstallBundleUtils.h
+++ b/Engine/Source/Runtime/InstallBundleManager/Public/InstallBundleUtils.h
@@ -8,7 +8,7 @@
 
 namespace InstallBundleUtil
 {
-	bool HasInternetConnection(ENetworkConnectionType ConnectionType);
+	INSTALLBUNDLEMANAGER_API bool HasInternetConnection(ENetworkConnectionType ConnectionType);
 
 	// It would really be nice to have these in core
 	template<class EnumType>
@@ -36,7 +36,7 @@ namespace InstallBundleUtil
 	}
 
 	// Keep the engine awake via RAII when running as an embedded app
-	class FInstallBundleManagerKeepAwake : public FEmbeddedKeepAwake
+	class INSTALLBUNDLEMANAGER_API FInstallBundleManagerKeepAwake : public FEmbeddedKeepAwake
 	{
 		static FName Tag;
 		static FName TagWithRendering;
@@ -45,7 +45,7 @@ namespace InstallBundleUtil
 			: FEmbeddedKeepAwake(bNeedsRendering ? TagWithRendering : Tag, bNeedsRendering) {}
 	};
 
-	class FInstallBundleWork : public FNonAbandonableTask
+	class INSTALLBUNDLEMANAGER_API FInstallBundleWork : public FNonAbandonableTask
 	{
 	public:
 		FInstallBundleWork() = default;
@@ -83,9 +83,9 @@ namespace InstallBundleUtil
 
 	using FInstallBundleTask = FAsyncTask<FInstallBundleWork>;
 
-	void StartInstallBundleAsyncIOTask(TArray<TUniquePtr<FInstallBundleTask>>& Tasks, TUniqueFunction<void()> WorkFunc, TUniqueFunction<void()> OnComplete);
+	INSTALLBUNDLEMANAGER_API void StartInstallBundleAsyncIOTask(TArray<TUniquePtr<FInstallBundleTask>>& Tasks, TUniqueFunction<void()> WorkFunc, TUniqueFunction<void()> OnComplete);
 
-	void FinishInstallBundleAsyncIOTasks(TArray<TUniquePtr<FInstallBundleTask>>& Tasks);
+	INSTALLBUNDLEMANAGER_API void FinishInstallBundleAsyncIOTasks(TArray<TUniquePtr<FInstallBundleTask>>& Tasks);
 
-	void CleanupInstallBundleAsyncIOTasks(TArray<TUniquePtr<FInstallBundleTask>>& Tasks);
+	INSTALLBUNDLEMANAGER_API void CleanupInstallBundleAsyncIOTasks(TArray<TUniquePtr<FInstallBundleTask>>& Tasks);
 }
