Install Bundle Manager Refactor to allow adding arbitray bundle sources
-Implement GetContentState for bundle sources
-Make builk build init actually async
-Temporarily convert watchdogs to templates to facilite refactoring

diff --git a/Engine/Source/Runtime/InstallBundleManager/Private/InstallBundleUtils.cpp b/Engine/Source/Runtime/InstallBundleManager/Private/InstallBundleUtils.cpp
--- a/Engine/Source/Runtime/InstallBundleManager/Private/InstallBundleUtils.cpp
+++ b/Engine/Source/Runtime/InstallBundleManager/Private/InstallBundleUtils.cpp
@@ -11,6 +11,16 @@ namespace InstallBundleUtil
 			&& ConnectionType != ENetworkConnectionType::None;
 	}
 
+	bool StateSignifiesNeedsInstall(EBundleState StateIn)
+	{
+		return (StateIn == EBundleState::NotInstalled || StateIn == EBundleState::NeedsUpdate);
+	}
+
+	bool StateSignifiesNeedsInstall(EInstallBundleContentState StateIn)
+	{
+		return (StateIn == EInstallBundleContentState::NotInstalled || StateIn == EInstallBundleContentState::NeedsUpdate);
+	}
+
 	FName FInstallBundleManagerKeepAwake::Tag(TEXT("InstallBundleManagerKeepAwake"));
 	FName FInstallBundleManagerKeepAwake::TagWithRendering(TEXT("InstallBundleManagerKeepAwakeWithRendering"));
 
diff --git a/Engine/Source/Runtime/InstallBundleManager/Public/InstallBundleManagerInterface.h b/Engine/Source/Runtime/InstallBundleManager/Public/InstallBundleManagerInterface.h
--- a/Engine/Source/Runtime/InstallBundleManager/Public/InstallBundleManagerInterface.h
+++ b/Engine/Source/Runtime/InstallBundleManager/Public/InstallBundleManagerInterface.h
@@ -201,8 +201,6 @@ DECLARE_DELEGATE_RetVal_OneParam(EInstallBundleManagerInitErrorHandlerResult, FI
 DECLARE_MULTICAST_DELEGATE_OneParam(FInstallBundleCompleteMultiDelegate, FInstallBundleResultInfo);
 DECLARE_MULTICAST_DELEGATE_OneParam(FInstallBundlePausedMultiDelegate, FInstallBundlePauseInfo);
 
-DECLARE_DELEGATE_OneParam(FInstallBundleGetContentStateDelegate, FInstallBundleContentState);
-
 class INSTALLBUNDLEMANAGER_API IInstallBundleManager
 {
 public:
diff --git a/Engine/Source/Runtime/InstallBundleManager/Public/InstallBundleSourceInterface.h b/Engine/Source/Runtime/InstallBundleManager/Public/InstallBundleSourceInterface.h
--- a/Engine/Source/Runtime/InstallBundleManager/Public/InstallBundleSourceInterface.h
+++ b/Engine/Source/Runtime/InstallBundleManager/Public/InstallBundleSourceInterface.h
@@ -15,6 +15,7 @@ public:
 	virtual ~IInstallBundleSource() {}
 
 	virtual EInstallBundleSourceType GetSourceType() const = 0;
+	virtual float GetSourceWeight() const { return 1.0f; }
 
 	virtual void Init(
 		TSharedRef<FConfigFile> InstallBundleConfig, 
@@ -25,5 +26,7 @@ public:
 	// Only valid after AsyncInit completes
 	virtual EInstallBundleManagerInitResult GetLastInitResult() const = 0;
 
+	virtual void GetContentState(TArrayView<FName> BundleNames, EInstallBundleGetContentStateFlags Flags, FInstallBundleGetContentStateDelegate Callback) = 0;
+
 	virtual void SetErrorSimulationCommands(const FString& CommandLine) {}
 };
diff --git a/Engine/Source/Runtime/InstallBundleManager/Public/InstallBundleTypes.h b/Engine/Source/Runtime/InstallBundleManager/Public/InstallBundleTypes.h
--- a/Engine/Source/Runtime/InstallBundleManager/Public/InstallBundleTypes.h
+++ b/Engine/Source/Runtime/InstallBundleManager/Public/InstallBundleTypes.h
@@ -78,3 +78,5 @@ enum class EInstallBundleGetContentStateFlags : uint32
 };
 ENUM_CLASS_FLAGS(EInstallBundleGetContentStateFlags);
 
+DECLARE_DELEGATE_OneParam(FInstallBundleGetContentStateDelegate, FInstallBundleContentState);
+
diff --git a/Engine/Source/Runtime/InstallBundleManager/Public/InstallBundleUtils.h b/Engine/Source/Runtime/InstallBundleManager/Public/InstallBundleUtils.h
--- a/Engine/Source/Runtime/InstallBundleManager/Public/InstallBundleUtils.h
+++ b/Engine/Source/Runtime/InstallBundleManager/Public/InstallBundleUtils.h
@@ -3,6 +3,7 @@
 #pragma once
 
 #include "CoreMinimal.h"
+#include "InstallBundleTypes.h"
 #include "Async/AsyncWork.h"
 #include "Misc/EmbeddedCommunication.h"
 
@@ -10,6 +11,10 @@ namespace InstallBundleUtil
 {
 	INSTALLBUNDLEMANAGER_API bool HasInternetConnection(ENetworkConnectionType ConnectionType);
 
+	INSTALLBUNDLEMANAGER_API bool StateSignifiesNeedsInstall(EBundleState StateIn);
+
+	INSTALLBUNDLEMANAGER_API bool StateSignifiesNeedsInstall(EInstallBundleContentState StateIn);
+
 	// It would really be nice to have these in core
 	template<class EnumType>
 	constexpr auto& CastAsUnderlying(EnumType &Type)
