Stub out most of the Install Bundle Source Interface

diff --git a/Engine/Source/Runtime/InstallBundleManager/Private/BundlePrereqCombinedStatusHelper.cpp b/Engine/Source/Runtime/InstallBundleManager/Private/BundlePrereqCombinedStatusHelper.cpp
--- a/Engine/Source/Runtime/InstallBundleManager/Private/BundlePrereqCombinedStatusHelper.cpp
+++ b/Engine/Source/Runtime/InstallBundleManager/Private/BundlePrereqCombinedStatusHelper.cpp
@@ -320,7 +320,7 @@ const FBundlePrereqCombinedStatusHelper::FCombinedBundleStatus& FBundlePrereqCom
 	return CurrentCombinedStatus;
 }
 
-void FBundlePrereqCombinedStatusHelper::OnBundleInstallComplete(FInstallBundleResultInfo CompletedBundleInfo)
+void FBundlePrereqCombinedStatusHelper::OnBundleInstallComplete(FInstallBundleRequestResultInfo CompletedBundleInfo)
 {
 	const FName CompletedBundleName = CompletedBundleInfo.BundleName;
 	const bool bBundleCompletedSuccessfully = (CompletedBundleInfo.Result == EInstallBundleResult::OK);
diff --git a/Engine/Source/Runtime/InstallBundleManager/Private/InstallBundleTypes.cpp b/Engine/Source/Runtime/InstallBundleManager/Private/InstallBundleTypes.cpp
--- a/Engine/Source/Runtime/InstallBundleManager/Private/InstallBundleTypes.cpp
+++ b/Engine/Source/Runtime/InstallBundleManager/Private/InstallBundleTypes.cpp
@@ -63,3 +63,20 @@ const TCHAR* LexToString(EInstallBundleContentState State)
 
 	return Strings[InstallBundleUtil::CastToUnderlying(State)];
 }
+
+const TCHAR* LexToString(EInstallBundleResult Result)
+{
+	static const TCHAR* Strings[] =
+	{
+		TEXT("OK"),
+		TEXT("FailedPrereqRequiresLatestClient"),
+		TEXT("InstallError"),
+		TEXT("InstallerOutOfDiskSpaceError"),
+		TEXT("ManifestArchiveError"),
+		TEXT("UserCancelledError"),
+		TEXT("InitializationError"),
+	};
+	static_assert(InstallBundleUtil::CastToUnderlying(EInstallBundleResult::Count) == UE_ARRAY_COUNT(Strings), "");
+
+	return Strings[InstallBundleUtil::CastToUnderlying(Result)];
+}
diff --git a/Engine/Source/Runtime/InstallBundleManager/Private/InstallBundleUtils.cpp b/Engine/Source/Runtime/InstallBundleManager/Private/InstallBundleUtils.cpp
--- a/Engine/Source/Runtime/InstallBundleManager/Private/InstallBundleUtils.cpp
+++ b/Engine/Source/Runtime/InstallBundleManager/Private/InstallBundleUtils.cpp
@@ -21,6 +21,22 @@ namespace InstallBundleUtil
 		return (StateIn == EInstallBundleContentState::NotInstalled || StateIn == EInstallBundleContentState::NeedsUpdate);
 	}
 
+	const TCHAR* GetInstallBundlePauseReason(EInstallBundlePauseFlags Flags)
+	{
+		// Return the most appropriate reason given the flags
+
+		if (EnumHasAnyFlags(Flags, EInstallBundlePauseFlags::UserPaused))
+			return TEXT("UserPaused");
+
+		if (EnumHasAnyFlags(Flags, EInstallBundlePauseFlags::NoInternetConnection))
+			return TEXT("NoInternetConnection");
+
+		if (EnumHasAnyFlags(Flags, EInstallBundlePauseFlags::OnCellularNetwork))
+			return TEXT("OnCellularNetwork");
+
+		return TEXT("");
+	}
+
 	FName FInstallBundleManagerKeepAwake::Tag(TEXT("InstallBundleManagerKeepAwake"));
 	FName FInstallBundleManagerKeepAwake::TagWithRendering(TEXT("InstallBundleManagerKeepAwakeWithRendering"));
 
diff --git a/Engine/Source/Runtime/InstallBundleManager/Public/BundlePrereqCombinedStatusHelper.h b/Engine/Source/Runtime/InstallBundleManager/Public/BundlePrereqCombinedStatusHelper.h
--- a/Engine/Source/Runtime/InstallBundleManager/Public/BundlePrereqCombinedStatusHelper.h
+++ b/Engine/Source/Runtime/InstallBundleManager/Public/BundlePrereqCombinedStatusHelper.h
@@ -61,7 +61,7 @@ private:
 	void CleanUpDelegates();
 	
 	//Called so we can track when a bundle is finished
-	void OnBundleInstallComplete(FInstallBundleResultInfo CompletedBundleInfo);
+	void OnBundleInstallComplete(FInstallBundleRequestResultInfo CompletedBundleInfo);
 	
 	float GetCombinedProgressPercent();
 	float GetIndividualWeightedProgressPercent(FInstallBundleStatus& Bundle);
diff --git a/Engine/Source/Runtime/InstallBundleManager/Public/InstallBundleManagerInterface.h b/Engine/Source/Runtime/InstallBundleManager/Public/InstallBundleManagerInterface.h
--- a/Engine/Source/Runtime/InstallBundleManager/Public/InstallBundleManagerInterface.h
+++ b/Engine/Source/Runtime/InstallBundleManager/Public/InstallBundleManagerInterface.h
@@ -10,72 +10,6 @@
 
 class IAnalyticsProviderET;
 
-enum class EInstallBundleResult : int
-{
-	OK,
-	FailedPrereqRequiresLatestClient,
-	InstallError,
-	InstallerOutOfDiskSpaceError,
-	ManifestArchiveError,
-	UserCancelledError,
-	InitializationError,
-	Count,
-};
-
-inline const TCHAR* LexToString(EInstallBundleResult Result)
-{
-	using UnderType = __underlying_type(EInstallBundleResult);
-	static const TCHAR* Strings[] =
-	{
-		TEXT("OK"),
-		TEXT("FailedPrereqRequiresLatestClient"),
-		TEXT("InstallError"),
-		TEXT("InstallerOutOfDiskSpaceError"),
-		TEXT("ManifestArchiveError"),
-		TEXT("UserCancelledError"),
-		TEXT("InitializationError"),
-	};
-	static_assert(static_cast<UnderType>(EInstallBundleResult::Count) == UE_ARRAY_COUNT(Strings), "");
-
-	return Strings[static_cast<UnderType>(Result)];
-}
-
-enum class EInstallBundlePauseFlags : uint32
-{
-	None = 0,
-	OnCellularNetwork = (1 << 0),
-	NoInternetConnection = (1 << 1),
-	UserPaused	= (1 << 2)
-};
-ENUM_CLASS_FLAGS(EInstallBundlePauseFlags)
-
-inline const TCHAR* GetInstallBundlePauseReason(EInstallBundlePauseFlags Flags)
-{
-	// Return the most appropriate reason given the flags
-
-	if (EnumHasAnyFlags(Flags, EInstallBundlePauseFlags::UserPaused))
-		return TEXT("UserPaused");
-
-	if (EnumHasAnyFlags(Flags, EInstallBundlePauseFlags::NoInternetConnection))
-		return TEXT("NoInternetConnection");
-	
-	if (EnumHasAnyFlags(Flags, EInstallBundlePauseFlags::OnCellularNetwork))
-		return TEXT("OnCellularNetwork");
-
-	return TEXT("");
-}
-
-enum class EInstallBundleRequestFlags : uint32
-{
-	None											= 0,
-	CheckForCellularDataUsage						= (1 << 0),
-	UseBackgroundDownloads							= (1 << 1),
-	SendNotificationIfDownloadCompletesInBackground = (1 << 2),
-	ForceNoPatching									= (1 << 3),
-	Defaults = UseBackgroundDownloads,
-};
-ENUM_CLASS_FLAGS(EInstallBundleRequestFlags)
-
 enum class EInstallBundleStatus : int
 {
 	Requested,
@@ -147,7 +81,7 @@ struct FInstallBundleStatus
 	float Finishing_Percent = 0;
 };
 
-struct FInstallBundleResultInfo
+struct FInstallBundleRequestResultInfo
 {
 	FName BundleName;
 	EInstallBundleResult Result = EInstallBundleResult::OK;
@@ -164,31 +98,6 @@ struct FInstallBundlePauseInfo
 	EInstallBundlePauseFlags PauseFlags = EInstallBundlePauseFlags::None;
 };
 
-enum class EInstallBundleRequestInfoFlags : int32
-{
-	None							= 0,
-	EnqueuedBundlesForInstall		= (1 << 0),
-	EnqueuedBundlesForRemoval		= (1 << 1),
-	SkippedAlreadyMountedBundles	= (1 << 2),
-	SkippedBundlesQueuedForInstall  = (1 << 3), // Only valid for removal requests
-	SkippedUnknownBundles			= (1 << 4),
-	InitializationError				= (1 << 5), // Can't enqueue because the bundle manager failed to initialize
-};
-ENUM_CLASS_FLAGS(EInstallBundleRequestInfoFlags);
-
-struct FInstallBundleRequestInfo
-{
-	EInstallBundleRequestInfoFlags InfoFlags = EInstallBundleRequestInfoFlags::None;
-	TArray<FName> BundlesQueuedForInstall;
-};
-
-enum class EInstallBundleCancelFlags : int32
-{
-	None		= 0,
-	Resumable	= (1 << 0),
-};
-ENUM_CLASS_FLAGS(EInstallBundleCancelFlags);
-
 enum class EInstallBundleManagerInitErrorHandlerResult
 {
 	NotHandled, // Defer to the next handler
@@ -198,7 +107,7 @@ enum class EInstallBundleManagerInitErrorHandlerResult
 
 DECLARE_DELEGATE_RetVal_OneParam(EInstallBundleManagerInitErrorHandlerResult, FInstallBundleManagerInitErrorHandler, EInstallBundleManagerInitResult);
 
-DECLARE_MULTICAST_DELEGATE_OneParam(FInstallBundleCompleteMultiDelegate, FInstallBundleResultInfo);
+DECLARE_MULTICAST_DELEGATE_OneParam(FInstallBundleCompleteMultiDelegate, FInstallBundleRequestResultInfo);
 DECLARE_MULTICAST_DELEGATE_OneParam(FInstallBundlePausedMultiDelegate, FInstallBundlePauseInfo);
 
 class INSTALLBUNDLEMANAGER_API IInstallBundleManager
diff --git a/Engine/Source/Runtime/InstallBundleManager/Public/InstallBundleSourceInterface.h b/Engine/Source/Runtime/InstallBundleManager/Public/InstallBundleSourceInterface.h
--- a/Engine/Source/Runtime/InstallBundleManager/Public/InstallBundleSourceInterface.h
+++ b/Engine/Source/Runtime/InstallBundleManager/Public/InstallBundleSourceInterface.h
@@ -9,24 +9,77 @@ class IAnalyticsProviderET;
 
 DECLARE_DELEGATE_ThreeParams(FInstallBundleSourceInitDelegate, TSharedRef<IInstallBundleSource> /*Source*/, EInstallBundleManagerInitResult /*Result*/, bool /*bShouldUseFallbackSource*/)
 
+DECLARE_DELEGATE_TwoParams(FInstallBundleCompleteDelegate, TSharedRef<IInstallBundleSource> /*Source*/, FInstallBundleSourceRequestResultInfo /*Result*/);
+
 class IInstallBundleSource : public TSharedFromThis<IInstallBundleSource>
 {
 public:
 	virtual ~IInstallBundleSource() {}
 
+	// Returns a unique id for this source
 	virtual EInstallBundleSourceType GetSourceType() const = 0;
+
+	// Returns the how this source should be weighted when combined with other sources
 	virtual float GetSourceWeight() const { return 1.0f; }
 
+	// Returns true if this source can remove content on init
+	virtual bool SupportsRemoveContentOnNextInit() const { return false; }
+
+	// Returns true if this source can cancel requests
+	virtual bool SupportsCanceling() const { return false; }
+
+	// Returns true if this source can pause requests
+	virtual bool SupportsPausing() const { return false; }
+
+	// Called once by bundle manager after constructing the bundle source
 	virtual void Init(
 		TSharedRef<FConfigFile> InstallBundleConfig, 
 		TSharedPtr<IAnalyticsProviderET> AnalyticsProvider) = 0;
+	// Bundle manager will not call AsyncInit again until the bundle source calls back that it is complete
+	// It will be retried indefinitely until init is successful.  
 	virtual void AsyncInit(FInstallBundleSourceInitDelegate Callback) = 0;
 
+	// Whether this source has been initialized or not
 	virtual EInstallBundleManagerInitState GetInitState() const = 0;
 	// Only valid after AsyncInit completes
+	// Returns the result of the last initialization attempt
 	virtual EInstallBundleManagerInitResult GetLastInitResult() const = 0;
 
+	// Gets the state of content on disk
+	// BundleNames contains all dependencies and has been deduped
 	virtual void GetContentState(TArrayView<FName> BundleNames, EInstallBundleGetContentStateFlags Flags, FInstallBundleGetContentStateDelegate Callback) = 0;
 
+	// Updates content on disk if necessary
+	// BundleNames contains all dependencies and has been deduped
+	//virtual FInstallBundleRequestInfo RequestUpdateContent(TArrayView<FName> BundleNames, EInstallBundleRequestFlags Flags, FInstallBundleCompleteDelegate Callbck) = 0;
+
+	// Returns true if content is scheduled to be removed the next time the source is initialized
+	// BundleNames contains all dependencies and has been deduped
+	virtual bool RequestRemoveContentOnNextInit(TArrayView<FName> RemoveNames) { return false; }
+	// Call to cancel the removal of any content scheduled for removal the next time the source is initialized
+	// Returns true if all bundles were canceled
+	virtual bool CancelRequestRemoveContentOnNextInit(TArrayView<FName> BundleNames) { return false;  }
+
+	// Cancel the install for the specified bundles
+	// BundleNames contains all dependencies and has been deduped
+	virtual void CancelBundles(TArrayView<FName> BundleNames, EInstallBundleCancelFlags Flags) {}
+
+	// Pause/Resume bundles.  Returns true if any bundles were paused/resumed
+	virtual bool PauseBundles(TArrayView<FName> BundleNames) { return false; }
+	virtual bool ResumeBundles(TArrayView<FName> BundleNames) { return false; }
+
+	// Called by bundle manager to update pause status.  Returns a struct indicating current pause status and if that status has changed.
+	virtual FInstallBundleSourcePauseInfo UpdateBundlePauseInfo(FName BundleName) 
+	{ 
+		FInstallBundleSourcePauseInfo Info;
+		Info.BundleName = BundleName;
+		return Info;
+	}
+
+	// TODO: UpdateContentRequestFlags - Probably need to rework the whole thing.  The only flag it makes sense for is CheckForCellularDataUsage
+
+	// TODO: GetBundleProgress() - Not sure how this will work yet
+
+	// Called by bundle manager to pass through command line options to simulate errors
 	virtual void SetErrorSimulationCommands(const FString& CommandLine) {}
 };
diff --git a/Engine/Source/Runtime/InstallBundleManager/Public/InstallBundleTypes.h b/Engine/Source/Runtime/InstallBundleManager/Public/InstallBundleTypes.h
--- a/Engine/Source/Runtime/InstallBundleManager/Public/InstallBundleTypes.h
+++ b/Engine/Source/Runtime/InstallBundleManager/Public/InstallBundleTypes.h
@@ -80,3 +80,77 @@ ENUM_CLASS_FLAGS(EInstallBundleGetContentStateFlags);
 
 DECLARE_DELEGATE_OneParam(FInstallBundleGetContentStateDelegate, FInstallBundleContentState);
 
+enum class EInstallBundleRequestInfoFlags : int32
+{
+	None = 0,
+	EnqueuedBundlesForInstall = (1 << 0),
+	SkippedAlreadyMountedBundles = (1 << 1),
+	SkippedBundlesQueuedForInstall = (1 << 2), // Only valid for removal requests
+	SkippedUnknownBundles = (1 << 3),
+	InitializationError = (1 << 4), // Can't enqueue because the bundle manager failed to initialize
+};
+ENUM_CLASS_FLAGS(EInstallBundleRequestInfoFlags);
+
+enum class EInstallBundleResult : int
+{
+	OK,
+	FailedPrereqRequiresLatestClient,
+	InstallError,
+	InstallerOutOfDiskSpaceError,
+	ManifestArchiveError,
+	UserCancelledError,
+	InitializationError,
+	Count,
+};
+INSTALLBUNDLEMANAGER_API const TCHAR* LexToString(EInstallBundleResult Result);
+
+enum class EInstallBundleRequestFlags : uint32
+{
+	None = 0,
+	CheckForCellularDataUsage = (1 << 0),
+	UseBackgroundDownloads = (1 << 1),
+	SendNotificationIfDownloadCompletesInBackground = (1 << 2),
+	ForceNoPatching = (1 << 3),
+	Defaults = UseBackgroundDownloads,
+};
+ENUM_CLASS_FLAGS(EInstallBundleRequestFlags)
+
+struct FInstallBundleRequestInfo
+{
+	EInstallBundleRequestInfoFlags InfoFlags = EInstallBundleRequestInfoFlags::None;
+	TArray<FName> BundlesQueuedForInstall;
+};
+
+struct FInstallBundleSourceRequestResultInfo
+{
+	FName BundleName;
+	EInstallBundleResult Result = EInstallBundleResult::OK;
+
+	// Forward any errors from the underlying implementation for a specific source
+	// Currently, these just forward BPT Error info
+	FText OptionalErrorText;
+	FString OptionalErrorCode;
+};
+
+enum class EInstallBundleCancelFlags : int32
+{
+	None = 0,
+	Resumable = (1 << 0),
+};
+ENUM_CLASS_FLAGS(EInstallBundleCancelFlags);
+
+enum class EInstallBundlePauseFlags : uint32
+{
+	None = 0,
+	OnCellularNetwork = (1 << 0),
+	NoInternetConnection = (1 << 1),
+	UserPaused = (1 << 2)
+};
+ENUM_CLASS_FLAGS(EInstallBundlePauseFlags);
+
+struct FInstallBundleSourcePauseInfo
+{
+	FName BundleName;
+	EInstallBundlePauseFlags PauseFlags = EInstallBundlePauseFlags::None;
+	bool bDidPauseChange = false;
+};
diff --git a/Engine/Source/Runtime/InstallBundleManager/Public/InstallBundleUtils.h b/Engine/Source/Runtime/InstallBundleManager/Public/InstallBundleUtils.h
--- a/Engine/Source/Runtime/InstallBundleManager/Public/InstallBundleUtils.h
+++ b/Engine/Source/Runtime/InstallBundleManager/Public/InstallBundleUtils.h
@@ -15,6 +15,8 @@ namespace InstallBundleUtil
 
 	INSTALLBUNDLEMANAGER_API bool StateSignifiesNeedsInstall(EInstallBundleContentState StateIn);
 
+	INSTALLBUNDLEMANAGER_API const TCHAR* GetInstallBundlePauseReason(EInstallBundlePauseFlags Flags);
+
 	// It would really be nice to have these in core
 	template<class EnumType>
 	constexpr auto& CastAsUnderlying(EnumType &Type)
diff --git a/Engine/Source/Runtime/Launch/Private/LaunchEngineLoop.cpp b/Engine/Source/Runtime/Launch/Private/LaunchEngineLoop.cpp
--- a/Engine/Source/Runtime/Launch/Private/LaunchEngineLoop.cpp
+++ b/Engine/Source/Runtime/Launch/Private/LaunchEngineLoop.cpp
@@ -926,7 +926,7 @@ void LaunchUpdateMostRecentProjectFile()
 }
 
 #if WITH_ENGINE
-void OnStartupContentMounted(FInstallBundleResultInfo Result, bool bDumpEarlyConfigReads, bool bDumpEarlyPakFileReads, bool bReloadConfig, bool bForceQuitAfterEarlyReads);
+void OnStartupContentMounted(FInstallBundleRequestResultInfo Result, bool bDumpEarlyConfigReads, bool bDumpEarlyPakFileReads, bool bReloadConfig, bool bForceQuitAfterEarlyReads);
 #endif
 void DumpEarlyReads(bool bDumpEarlyConfigReads, bool bDumpEarlyPakFileReads, bool bForceQuitAfterEarlyReads);
 void HandleConfigReload(bool bReloadConfig);
@@ -3975,7 +3975,7 @@ void FEngineLoop::ProcessLocalPlayerSlateOperations() const
 }
 
 #if WITH_ENGINE
-void OnStartupContentMounted(FInstallBundleResultInfo Result, bool bDumpEarlyConfigReads, bool bDumpEarlyPakFileReads, bool bReloadConfig, bool bForceQuitAfterEarlyReads)
+void OnStartupContentMounted(FInstallBundleRequestResultInfo Result, bool bDumpEarlyConfigReads, bool bDumpEarlyPakFileReads, bool bReloadConfig, bool bForceQuitAfterEarlyReads)
 {
 	if (Result.bIsStartup && Result.Result == EInstallBundleResult::OK)
 	{
