Fix FBundlePrereqCombinedStatusHelper tracking all bundles returned in a FInstallBundleContentState instead of only the ones actually being installed.
Move bundle weighting range normalization into FBundlePrereqCombinedStatusHelper so bundles are wegihted correctly with this fix.

diff --git a/Engine/Source/Runtime/InstallBundleManager/Private/BundlePrereqCombinedStatusHelper.cpp b/Engine/Source/Runtime/InstallBundleManager/Private/BundlePrereqCombinedStatusHelper.cpp
--- a/Engine/Source/Runtime/InstallBundleManager/Private/BundlePrereqCombinedStatusHelper.cpp
+++ b/Engine/Source/Runtime/InstallBundleManager/Private/BundlePrereqCombinedStatusHelper.cpp
@@ -100,25 +100,48 @@ void FBundlePrereqCombinedStatusHelper::CleanUpDelegates()
 	}
 }
 
-void FBundlePrereqCombinedStatusHelper::SetBundlesToTrackFromContentState(const FInstallBundleContentState& BundleContentState)
+void FBundlePrereqCombinedStatusHelper::SetBundlesToTrackFromContentState(const FInstallBundleContentState& BundleContentState, TArrayView<FName> BundlesToTrack)
 {
 	RequiredBundleNames.Empty();
 	CachedBundleWeights.Empty();
 	BundleStatusCache.Empty();
 	
-	//Track if we need any kind of bundle updates
 	bBundleNeedsUpdate = false;
-	if (BundleContentState.State == EInstallBundleContentState::NotInstalled || BundleContentState.State == EInstallBundleContentState::NeedsUpdate)
+	float TotalWeight = 0.0f;
+	for (const FName& Bundle : BundlesToTrack)
 	{
-		bBundleNeedsUpdate = true;
+		const EInstallBundleContentState* BundleState = BundleContentState.IndividualBundleStates.Find(Bundle);
+		const float* BundleWeight = BundleContentState.IndividualBundleWeights.Find(Bundle);
+		if (ensureAlwaysMsgf(BundleState && BundleWeight, TEXT("Trying to track unknown bundle %s"), *Bundle.ToString()))
+		{
+			//Track if we need any kind of bundle updates
+			if (*BundleState == EInstallBundleContentState::NotInstalled || *BundleState == EInstallBundleContentState::NeedsUpdate)
+			{
+				bBundleNeedsUpdate = true;
+			}
+
+			//Save required bundles and their weights
+			RequiredBundleNames.Add(Bundle);
+			CachedBundleWeights.FindOrAdd(Bundle) = *BundleWeight;
+			TotalWeight += *BundleWeight;
+		}
 	}
+
 	CurrentCombinedStatus.bBundleRequiresUpdate = bBundleNeedsUpdate;
 	
-	//Save required bundles and their weights
-	for (const TPair<FName, float>& BundleStatePair : BundleContentState.IndividualBundleWeights)
+	if (TotalWeight > 0.0f)
 	{
-		RequiredBundleNames.Add(BundleStatePair.Key);
-		CachedBundleWeights.FindOrAdd(BundleStatePair.Key) = BundleStatePair.Value;
+		for (TPair<FName, float>& BundleWeightPair : CachedBundleWeights)
+		{
+			BundleWeightPair.Value /= TotalWeight;
+		}
+	}
+
+	// If no bundles to track, we are done
+	if (RequiredBundleNames.Num() == 0)
+	{
+		CurrentCombinedStatus.ProgressPercent = 1.0f;
+		CurrentCombinedStatus.CombinedState = FCombinedBundleStatus::ECombinedBundleStateEnum::Finished;
 	}
 	
 	//Go ahead and calculate initial values from the Bundle Cache
@@ -151,6 +174,9 @@ void FBundlePrereqCombinedStatusHelper::UpdateBundleCache()
 
 void FBundlePrereqCombinedStatusHelper::UpdateCombinedStatus()
 {
+	if (RequiredBundleNames.Num() == 0)
+		return;
+
 	CurrentCombinedStatus.ProgressPercent = GetCombinedProgressPercent();
 	
 	EInstallBundleStatus EarliestBundleState = EInstallBundleStatus::Count;
@@ -235,15 +261,15 @@ void FBundlePrereqCombinedStatusHelper::UpdateCombinedStatus()
 	}
 }
 
-float FBundlePrereqCombinedStatusHelper::GetCombinedProgressPercent()
+float FBundlePrereqCombinedStatusHelper::GetCombinedProgressPercent() const
 {
 	float AllBundleProgressPercent = 0.f;
 	
 	ensureAlwaysMsgf((CachedBundleWeights.Num() >= BundleStatusCache.Num()), TEXT("Missing Cache Entries for BundleWeights!Any missing bundles will have 0 for their progress!"));
 	
-	for (TPair<FName,FInstallBundleStatus>& BundleStatusPair : BundleStatusCache)
+	for (const TPair<FName,FInstallBundleStatus>& BundleStatusPair : BundleStatusCache)
 	{
-		float* FoundWeight = CachedBundleWeights.Find(BundleStatusPair.Key);
+		const float* FoundWeight = CachedBundleWeights.Find(BundleStatusPair.Key);
 		if (ensureAlwaysMsgf((nullptr != FoundWeight), TEXT("Found missing entry for BundleWeight! Bundle %s does not have a weight entry!"), *(BundleStatusPair.Key.ToString())))
 		{
 			AllBundleProgressPercent += ((*FoundWeight) * GetIndividualWeightedProgressPercent(BundleStatusPair.Value));
@@ -254,7 +280,7 @@ float FBundlePrereqCombinedStatusHelper::GetCombinedProgressPercent()
 	return AllBundleProgressPercent;
 }
 
-float FBundlePrereqCombinedStatusHelper::GetIndividualWeightedProgressPercent(FInstallBundleStatus& BundleStatus)
+float FBundlePrereqCombinedStatusHelper::GetIndividualWeightedProgressPercent(const FInstallBundleStatus& BundleStatus) const
 {
 	const float TotalWeight = DownloadWeight + InstallWeight;
 	float CombinedOverallProgressPercent = 0.f;
diff --git a/Engine/Source/Runtime/InstallBundleManager/Public/BundlePrereqCombinedStatusHelper.h b/Engine/Source/Runtime/InstallBundleManager/Public/BundlePrereqCombinedStatusHelper.h
--- a/Engine/Source/Runtime/InstallBundleManager/Public/BundlePrereqCombinedStatusHelper.h
+++ b/Engine/Source/Runtime/InstallBundleManager/Public/BundlePrereqCombinedStatusHelper.h
@@ -42,7 +42,7 @@ public:
 	FBundlePrereqCombinedStatusHelper& operator=(FBundlePrereqCombinedStatusHelper&& Other);
 	
 	//Setup tracking for all bundles required in the supplied BundleContentState
-	void SetBundlesToTrackFromContentState(const FInstallBundleContentState& BundleContentState);
+	void SetBundlesToTrackFromContentState(const FInstallBundleContentState& BundleContentState, TArrayView<FName> BundlesToTrack);
 	
 	//Get current CombinedBundleStatus for everything setup to track
 	const FCombinedBundleStatus& GetCurrentCombinedState() const;
@@ -63,8 +63,8 @@ private:
 	//Called so we can track when a bundle is finished
 	void OnBundleInstallComplete(FInstallBundleRequestResultInfo CompletedBundleInfo);
 	
-	float GetCombinedProgressPercent();
-	float GetIndividualWeightedProgressPercent(FInstallBundleStatus& Bundle);
+	float GetCombinedProgressPercent() const;
+	float GetIndividualWeightedProgressPercent(const FInstallBundleStatus& Bundle) const;
 	
 private:
 	//All bundles we need including pre-reqs
diff --git a/Engine/Source/Runtime/InstallBundleManager/Public/InstallBundleTypes.h b/Engine/Source/Runtime/InstallBundleManager/Public/InstallBundleTypes.h
--- a/Engine/Source/Runtime/InstallBundleManager/Public/InstallBundleTypes.h
+++ b/Engine/Source/Runtime/InstallBundleManager/Public/InstallBundleTypes.h
@@ -85,9 +85,8 @@ enum class EInstallBundleRequestInfoFlags : int32
 	None = 0,
 	EnqueuedBundlesForInstall = (1 << 0),
 	SkippedAlreadyMountedBundles = (1 << 1),
-	SkippedBundlesQueuedForInstall = (1 << 2), // Only valid for removal requests
-	SkippedUnknownBundles = (1 << 3),
-	InitializationError = (1 << 4), // Can't enqueue because the bundle manager failed to initialize
+	SkippedUnknownBundles = (1 << 2),
+	InitializationError = (1 << 3), // Can't enqueue because the bundle manager failed to initialize
 };
 ENUM_CLASS_FLAGS(EInstallBundleRequestInfoFlags);
 
