Install Bundle Manager
- Add Version Info to GetContentState

diff --git a/Engine/Source/Runtime/InstallBundleManager/Private/BundlePrereqCombinedStatusHelper.cpp b/Engine/Source/Runtime/InstallBundleManager/Private/BundlePrereqCombinedStatusHelper.cpp
--- a/Engine/Source/Runtime/InstallBundleManager/Private/BundlePrereqCombinedStatusHelper.cpp
+++ b/Engine/Source/Runtime/InstallBundleManager/Private/BundlePrereqCombinedStatusHelper.cpp
@@ -100,7 +100,7 @@ void FBundlePrereqCombinedStatusHelper::CleanUpDelegates()
 	}
 }
 
-void FBundlePrereqCombinedStatusHelper::SetBundlesToTrackFromContentState(const FInstallBundleContentState& BundleContentState, TArrayView<FName> BundlesToTrack)
+void FBundlePrereqCombinedStatusHelper::SetBundlesToTrackFromContentState(const FInstallBundleCombinedContentState& BundleContentState, TArrayView<FName> BundlesToTrack)
 {
 	RequiredBundleNames.Empty();
 	CachedBundleWeights.Empty();
@@ -110,20 +110,19 @@ void FBundlePrereqCombinedStatusHelper::SetBundlesToTrackFromContentState(const
 	float TotalWeight = 0.0f;
 	for (const FName& Bundle : BundlesToTrack)
 	{
-		const EInstallBundleContentState* BundleState = BundleContentState.IndividualBundleStates.Find(Bundle);
-		const float* BundleWeight = BundleContentState.IndividualBundleWeights.Find(Bundle);
-		if (ensureAlwaysMsgf(BundleState && BundleWeight, TEXT("Trying to track unknown bundle %s"), *Bundle.ToString()))
+		const FInstallBundleContentState* BundleState = BundleContentState.IndividualBundleStates.Find(Bundle);
+		if (ensureAlwaysMsgf(BundleState, TEXT("Trying to track unknown bundle %s"), *Bundle.ToString()))
 		{
 			//Track if we need any kind of bundle updates
-			if (*BundleState == EInstallBundleContentState::NotInstalled || *BundleState == EInstallBundleContentState::NeedsUpdate)
+			if (BundleState->State == EInstallBundleContentState::NotInstalled || BundleState->State == EInstallBundleContentState::NeedsUpdate)
 			{
 				bBundleNeedsUpdate = true;
 			}
 
 			//Save required bundles and their weights
 			RequiredBundleNames.Add(Bundle);
-			CachedBundleWeights.FindOrAdd(Bundle) = *BundleWeight;
-			TotalWeight += *BundleWeight;
+			CachedBundleWeights.FindOrAdd(Bundle) = BundleState->Weight;
+			TotalWeight += BundleState->Weight;
 		}
 	}
 
diff --git a/Engine/Source/Runtime/InstallBundleManager/Private/InstallBundleUtils.cpp b/Engine/Source/Runtime/InstallBundleManager/Private/InstallBundleUtils.cpp
--- a/Engine/Source/Runtime/InstallBundleManager/Private/InstallBundleUtils.cpp
+++ b/Engine/Source/Runtime/InstallBundleManager/Private/InstallBundleUtils.cpp
@@ -2,9 +2,15 @@
 
 #include "InstallBundleUtils.h"
 #include "InstallBundleManagerPrivatePCH.h"
+#include "Misc/App.h"
 
 namespace InstallBundleUtil
 {
+	INSTALLBUNDLEMANAGER_API FString GetAppVersion()
+	{
+		return FString::Printf(TEXT("%s-%s"), FApp::GetBuildVersion(), ANSI_TO_TCHAR(FPlatformProperties::IniPlatformName()));
+	}
+
 	bool HasInternetConnection(ENetworkConnectionType ConnectionType)
 	{
 		return ConnectionType != ENetworkConnectionType::AirplaneMode
diff --git a/Engine/Source/Runtime/InstallBundleManager/Public/BundlePrereqCombinedStatusHelper.h b/Engine/Source/Runtime/InstallBundleManager/Public/BundlePrereqCombinedStatusHelper.h
--- a/Engine/Source/Runtime/InstallBundleManager/Public/BundlePrereqCombinedStatusHelper.h
+++ b/Engine/Source/Runtime/InstallBundleManager/Public/BundlePrereqCombinedStatusHelper.h
@@ -42,7 +42,7 @@ public:
 	FBundlePrereqCombinedStatusHelper& operator=(FBundlePrereqCombinedStatusHelper&& Other);
 	
 	//Setup tracking for all bundles required in the supplied BundleContentState
-	void SetBundlesToTrackFromContentState(const FInstallBundleContentState& BundleContentState, TArrayView<FName> BundlesToTrack);
+	void SetBundlesToTrackFromContentState(const FInstallBundleCombinedContentState& BundleContentState, TArrayView<FName> BundlesToTrack);
 	
 	//Get current CombinedBundleStatus for everything setup to track
 	const FCombinedBundleStatus& GetCurrentCombinedState() const;
diff --git a/Engine/Source/Runtime/InstallBundleManager/Public/InstallBundleTypes.h b/Engine/Source/Runtime/InstallBundleManager/Public/InstallBundleTypes.h
--- a/Engine/Source/Runtime/InstallBundleManager/Public/InstallBundleTypes.h
+++ b/Engine/Source/Runtime/InstallBundleManager/Public/InstallBundleTypes.h
@@ -63,8 +63,15 @@ INSTALLBUNDLEMANAGER_API const TCHAR* LexToString(EInstallBundleContentState Sta
 struct INSTALLBUNDLEMANAGER_API FInstallBundleContentState
 {
 	EInstallBundleContentState State = EInstallBundleContentState::InitializationError;
-	TMap<FName, EInstallBundleContentState> IndividualBundleStates;
-	TMap<FName, float> IndividualBundleWeights;
+	float Weight = 0.0f;
+	TMap<EInstallBundleSourceType, FString> Version;
+};
+
+struct INSTALLBUNDLEMANAGER_API FInstallBundleCombinedContentState
+{
+	EInstallBundleContentState State = EInstallBundleContentState::InitializationError;
+	TMap<FName, FInstallBundleContentState> IndividualBundleStates;
+	TMap<EInstallBundleSourceType, FString> CurrentVersion;
 	uint64 DownloadSize = 0;
 	uint64 InstallSize = 0;
 	uint64 InstallOverheadSize = 0;
@@ -78,7 +85,7 @@ enum class EInstallBundleGetContentStateFlags : uint32
 };
 ENUM_CLASS_FLAGS(EInstallBundleGetContentStateFlags);
 
-DECLARE_DELEGATE_OneParam(FInstallBundleGetContentStateDelegate, FInstallBundleContentState);
+DECLARE_DELEGATE_OneParam(FInstallBundleGetContentStateDelegate, FInstallBundleCombinedContentState);
 
 enum class EInstallBundleRequestInfoFlags : int32
 {
diff --git a/Engine/Source/Runtime/InstallBundleManager/Public/InstallBundleUtils.h b/Engine/Source/Runtime/InstallBundleManager/Public/InstallBundleUtils.h
--- a/Engine/Source/Runtime/InstallBundleManager/Public/InstallBundleUtils.h
+++ b/Engine/Source/Runtime/InstallBundleManager/Public/InstallBundleUtils.h
@@ -9,6 +9,9 @@
 
 namespace InstallBundleUtil
 {
+	// Returns the app version in the same format as BPS versions
+	INSTALLBUNDLEMANAGER_API FString GetAppVersion();
+
 	INSTALLBUNDLEMANAGER_API bool HasInternetConnection(ENetworkConnectionType ConnectionType);
 
 	INSTALLBUNDLEMANAGER_API bool StateSignifiesNeedsInstall(EBundleState StateIn);
diff --git a/Engine/Source/Runtime/NullInstallBundleManager/Source/NullInstallBundleManager.cpp b/Engine/Source/Runtime/NullInstallBundleManager/Source/NullInstallBundleManager.cpp
--- a/Engine/Source/Runtime/NullInstallBundleManager/Source/NullInstallBundleManager.cpp
+++ b/Engine/Source/Runtime/NullInstallBundleManager/Source/NullInstallBundleManager.cpp
@@ -46,14 +46,14 @@ class FNullInstallBundleManager : public IInstallBundleManager
 
 	virtual void GetContentState(FName BundleName, EInstallBundleGetContentStateFlags Flags, bool bAddDependencies, FInstallBundleGetContentStateDelegate Callback, FName RequestTag) override
 	{
-		FInstallBundleContentState State;
+		FInstallBundleCombinedContentState State;
 		State.State = EInstallBundleContentState::UpToDate;
 		Callback.ExecuteIfBound(State);
 	}
 
 	virtual void GetContentState(TArrayView<FName> BundleNames, EInstallBundleGetContentStateFlags Flags, bool bAddDependencies, FInstallBundleGetContentStateDelegate Callback, FName RequestTag) override
 	{
-		FInstallBundleContentState State;
+		FInstallBundleCombinedContentState State;
 		State.State = EInstallBundleContentState::UpToDate;
 		Callback.ExecuteIfBound(State);
 	}
