Add a flag for whether we installed anything to FInstallBundleRequestResultInfo.

diff --git a/Engine/Source/Runtime/InstallBundleManager/Public/InstallBundleManagerInterface.h b/Engine/Source/Runtime/InstallBundleManager/Public/InstallBundleManagerInterface.h
--- a/Engine/Source/Runtime/InstallBundleManager/Public/InstallBundleManagerInterface.h
+++ b/Engine/Source/Runtime/InstallBundleManager/Public/InstallBundleManagerInterface.h
@@ -86,6 +86,7 @@ struct FInstallBundleRequestResultInfo
 	FName BundleName;
 	EInstallBundleResult Result = EInstallBundleResult::OK;
 	bool bIsStartup = false;
+	bool bContentWasInstalled = false;
 
 	// Currently, these just forward BPT Error info
 	FText OptionalErrorText;
