Install Analytics
- Set a session ID which will be unique to a machine + content version.
- Make sure heartbeat analytics fires as soon as bundle manager starts doing work instead or a minute later.
- Make sure heartbeat watchdog stops ticking when there is no work.
- Fix background download watchdog not being set.

diff --git a/Engine/Source/Runtime/InstallBundleManager/Public/InstallBundleSourceInterface.h b/Engine/Source/Runtime/InstallBundleManager/Public/InstallBundleSourceInterface.h
--- a/Engine/Source/Runtime/InstallBundleManager/Public/InstallBundleSourceInterface.h
+++ b/Engine/Source/Runtime/InstallBundleManager/Public/InstallBundleSourceInterface.h
@@ -45,6 +45,9 @@ public:
 	// Returns the result of the last initialization attempt
 	virtual EInstallBundleManagerInitResult GetLastInitResult() const = 0;
 
+	// Returns content version in a "<BuildVersion>-<Platform>" format
+	virtual FString GetContentVersion() const = 0;
+
 	// Gets the state of content on disk
 	// BundleNames contains all dependencies and has been deduped
 	virtual void GetContentState(TArrayView<FName> BundleNames, EInstallBundleGetContentStateFlags Flags, FInstallBundleGetContentStateDelegate Callback) = 0;
