Install Bundle Manger Refactor - Content Request Interface - Checking in what I have so far for Hongyi.

diff --git a/Engine/Source/Runtime/InstallBundleManager/Private/InstallBundleUtils.cpp b/Engine/Source/Runtime/InstallBundleManager/Private/InstallBundleUtils.cpp
--- a/Engine/Source/Runtime/InstallBundleManager/Private/InstallBundleUtils.cpp
+++ b/Engine/Source/Runtime/InstallBundleManager/Private/InstallBundleUtils.cpp
@@ -4,6 +4,8 @@
 #include "InstallBundleManagerPrivatePCH.h"
 #include "Misc/App.h"
 
+#include "HAL/PlatformApplicationMisc.h"
+
 namespace InstallBundleUtil
 {
 	INSTALLBUNDLEMANAGER_API FString GetAppVersion()
@@ -46,6 +48,29 @@ namespace InstallBundleUtil
 	FName FInstallBundleManagerKeepAwake::Tag(TEXT("InstallBundleManagerKeepAwake"));
 	FName FInstallBundleManagerKeepAwake::TagWithRendering(TEXT("InstallBundleManagerKeepAwakeWithRendering"));
 
+	bool FInstallBundleManagerScreenSaverControl::bDidDisableScreensaver = false;
+	int FInstallBundleManagerScreenSaverControl::DisableCount = 0;
+
+	void FInstallBundleManagerScreenSaverControl::IncDisable()
+	{
+		if (!bDidDisableScreensaver && FPlatformApplicationMisc::IsScreensaverEnabled())
+		{
+			bDidDisableScreensaver = FPlatformApplicationMisc::ControlScreensaver(FGenericPlatformApplicationMisc::EScreenSaverAction::Disable);
+		}
+
+		++DisableCount;
+	}
+
+	void FInstallBundleManagerScreenSaverControl::DecDisable()
+	{
+		--DisableCount;
+		if (DisableCount == 0 && bDidDisableScreensaver)
+		{
+			FPlatformApplicationMisc::ControlScreensaver(FGenericPlatformApplicationMisc::EScreenSaverAction::Enable);
+			bDidDisableScreensaver = false;
+		}
+	}
+
 	void StartInstallBundleAsyncIOTask(TArray<TUniquePtr<FInstallBundleTask>>& Tasks, TUniqueFunction<void()> WorkFunc, TUniqueFunction<void()> OnComplete)
 	{
 		TUniquePtr<FInstallBundleTask> Task = MakeUnique<FInstallBundleTask>(MoveTemp(WorkFunc), MoveTemp(OnComplete));
diff --git a/Engine/Source/Runtime/InstallBundleManager/Public/InstallBundleSourceInterface.h b/Engine/Source/Runtime/InstallBundleManager/Public/InstallBundleSourceInterface.h
--- a/Engine/Source/Runtime/InstallBundleManager/Public/InstallBundleSourceInterface.h
+++ b/Engine/Source/Runtime/InstallBundleManager/Public/InstallBundleSourceInterface.h
@@ -9,7 +9,7 @@ class IAnalyticsProviderET;
 
 DECLARE_DELEGATE_ThreeParams(FInstallBundleSourceInitDelegate, TSharedRef<IInstallBundleSource> /*Source*/, EInstallBundleManagerInitResult /*Result*/, bool /*bShouldUseFallbackSource*/)
 
-DECLARE_DELEGATE_TwoParams(FInstallBundleCompleteDelegate, TSharedRef<IInstallBundleSource> /*Source*/, FInstallBundleSourceRequestResultInfo /*Result*/);
+DECLARE_DELEGATE_OneParam(FInstallBundleCompleteDelegate, FInstallBundleSourceRequestResultInfo /*Result*/);
 
 class IInstallBundleSource : public TSharedFromThis<IInstallBundleSource>
 {
@@ -54,7 +54,7 @@ public:
 
 	// Updates content on disk if necessary
 	// BundleNames contains all dependencies and has been deduped
-	//virtual FInstallBundleRequestInfo RequestUpdateContent(TArrayView<FName> BundleNames, EInstallBundleRequestFlags Flags, FInstallBundleCompleteDelegate Callbck) = 0;
+	virtual void RequestUpdateContent(FName BundleName, EInstallBundleRequestFlags Flags, FInstallBundleCompleteDelegate Callback) = 0;
 
 	// Returns true if content is scheduled to be removed the next time the source is initialized
 	// BundleNames contains all dependencies and has been deduped
diff --git a/Engine/Source/Runtime/InstallBundleManager/Public/InstallBundleTypes.h b/Engine/Source/Runtime/InstallBundleManager/Public/InstallBundleTypes.h
--- a/Engine/Source/Runtime/InstallBundleManager/Public/InstallBundleTypes.h
+++ b/Engine/Source/Runtime/InstallBundleManager/Public/InstallBundleTypes.h
@@ -136,6 +136,8 @@ struct FInstallBundleSourceRequestResultInfo
 	// Currently, these just forward BPT Error info
 	FText OptionalErrorText;
 	FString OptionalErrorCode;
+
+	TArray<FString> ContentPaths;
 };
 
 enum class EInstallBundleCancelFlags : int32
diff --git a/Engine/Source/Runtime/InstallBundleManager/Public/InstallBundleUtils.h b/Engine/Source/Runtime/InstallBundleManager/Public/InstallBundleUtils.h
--- a/Engine/Source/Runtime/InstallBundleManager/Public/InstallBundleUtils.h
+++ b/Engine/Source/Runtime/InstallBundleManager/Public/InstallBundleUtils.h
@@ -55,6 +55,39 @@ namespace InstallBundleUtil
 			: FEmbeddedKeepAwake(bNeedsRendering ? TagWithRendering : Tag, bNeedsRendering) {}
 	};
 
+	class INSTALLBUNDLEMANAGER_API FInstallBundleManagerScreenSaverControl
+	{
+		static bool bDidDisableScreensaver;
+		static int DisableCount;
+
+		static void IncDisable();
+		static void DecDisable();
+
+	public:
+		FInstallBundleManagerScreenSaverControl()
+		{
+			IncDisable();
+		}
+
+		~FInstallBundleManagerScreenSaverControl()
+		{
+			DecDisable();
+		}
+
+		FInstallBundleManagerScreenSaverControl(const FInstallBundleManagerScreenSaverControl& Other)
+		{
+			IncDisable();
+		}
+		FInstallBundleManagerScreenSaverControl(FInstallBundleManagerScreenSaverControl&& Other) = default;
+
+		FInstallBundleManagerScreenSaverControl& operator=(const FInstallBundleManagerScreenSaverControl& Other) = default;
+		FInstallBundleManagerScreenSaverControl& operator=(FInstallBundleManagerScreenSaverControl&& Other)
+		{
+			DecDisable();
+			return *this;
+		}
+	};
+
 	class INSTALLBUNDLEMANAGER_API FInstallBundleWork : public FNonAbandonableTask
 	{
 	public:
