Install Bundle Manger Refactor - Another snapshot for Hongyi.

Move bundle stats into a shared utility.
Update bundle source interface.
Begin implementing content requests.

diff --git a/Engine/Source/Runtime/InstallBundleManager/Private/InstallBundleUtils.cpp b/Engine/Source/Runtime/InstallBundleManager/Private/InstallBundleUtils.cpp
--- a/Engine/Source/Runtime/InstallBundleManager/Private/InstallBundleUtils.cpp
+++ b/Engine/Source/Runtime/InstallBundleManager/Private/InstallBundleUtils.cpp
@@ -112,4 +112,62 @@ namespace InstallBundleUtil
 			}
 		}
 	}
+
+	void FContentRequestStatsMap::StatsBegin(FName BundleName)
+	{
+		FContentRequestStats& Stats = StatsMap.FindOrAdd(BundleName);
+		if (ensureAlwaysMsgf(Stats.bOpen, TEXT("StatsBegin - Stat closed for %s"), *BundleName.ToString()) == false)
+		{
+			Stats = FContentRequestStats();
+		}
+
+		Stats.StartTime = FPlatformTime::Seconds();
+	}
+
+	void FContentRequestStatsMap::StatsEnd(FName BundleName)
+	{
+		FContentRequestStats& Stats = StatsMap.FindOrAdd(BundleName);
+
+		if (ensureAlwaysMsgf(Stats.bOpen, TEXT("StatsEnd - Stat closed for %s"), *BundleName.ToString()))
+		{
+			Stats.EndTime = FPlatformTime::Seconds();
+			Stats.bOpen = false;
+		}
+	}
+
+	void FContentRequestStatsMap::StatsBegin(FName BundleName, const TCHAR* State)
+	{
+		FContentRequestStats& Stats = StatsMap.FindOrAdd(BundleName);
+		if (ensureAlwaysMsgf(Stats.bOpen, TEXT("StatsBegin - Stat closed for %s - %s"), *BundleName.ToString(), State) == false)
+		{
+			Stats = FContentRequestStats();
+			Stats.StartTime = FPlatformTime::Seconds();
+		}
+
+		FContentRequestStateStats& StateStats = Stats.StateStats.FindOrAdd(State);
+		if (ensureAlwaysMsgf(StateStats.bOpen, TEXT("StatsBegin - StateStat closed for %s - %s"), *BundleName.ToString(), State) == false)
+		{
+			StateStats = FContentRequestStateStats();
+		}
+
+		StateStats.StartTime = FPlatformTime::Seconds();
+	}
+
+	void FContentRequestStatsMap::StatsEnd(FName BundleName, const TCHAR* State, uint64 DataSize /*= 0*/)
+	{
+		FContentRequestStats& Stats = StatsMap.FindOrAdd(BundleName);
+		if (ensureAlwaysMsgf(Stats.bOpen, TEXT("StatsEnd - Stat closed for %s - %s"), *BundleName.ToString(), State) == false)
+		{
+			Stats = FContentRequestStats();
+			Stats.StartTime = FPlatformTime::Seconds();
+		}
+
+		FContentRequestStateStats& StateStats = Stats.StateStats.FindOrAdd(State);
+		if(ensureAlwaysMsgf(StateStats.bOpen, TEXT("StatsEnd - StateStat closed for %s - %s"), *BundleName.ToString(), State))
+		{
+			StateStats.EndTime = FPlatformTime::Seconds();
+			StateStats.DataSize = DataSize;
+			StateStats.bOpen = false;
+		}
+	}
 }
diff --git a/Engine/Source/Runtime/InstallBundleManager/Public/InstallBundleSourceInterface.h b/Engine/Source/Runtime/InstallBundleManager/Public/InstallBundleSourceInterface.h
--- a/Engine/Source/Runtime/InstallBundleManager/Public/InstallBundleSourceInterface.h
+++ b/Engine/Source/Runtime/InstallBundleManager/Public/InstallBundleSourceInterface.h
@@ -3,6 +3,7 @@
 #pragma once
 
 #include "InstallBundleTypes.h"
+#include "InstallBundleUtils.h"
 
 class IInstallBundleSource;
 class IAnalyticsProviderET;
@@ -34,6 +35,7 @@ public:
 	// Called once by bundle manager after constructing the bundle source
 	virtual void Init(
 		TSharedRef<FConfigFile> InstallBundleConfig, 
+		TSharedRef<InstallBundleUtil::FContentRequestStatsMap> InRequestStats,
 		TSharedPtr<IAnalyticsProviderET> AnalyticsProvider) = 0;
 	// Bundle manager will not call AsyncInit again until the bundle source calls back that it is complete
 	// It will be retried indefinitely until init is successful.  
@@ -52,9 +54,16 @@ public:
 	// BundleNames contains all dependencies and has been deduped
 	virtual void GetContentState(TArrayView<FName> BundleNames, EInstallBundleGetContentStateFlags Flags, FInstallBundleGetContentStateDelegate Callback) = 0;
 
+	struct RequestUpdateContentBundleContext
+	{
+		FName BundleName;
+		EInstallBundleRequestFlags Flags = EInstallBundleRequestFlags::None;
+		FInstallBundleCompleteDelegate Callback;
+	};
+
 	// Updates content on disk if necessary
-	// BundleNames contains all dependencies and has been deduped
-	virtual void RequestUpdateContent(FName BundleName, EInstallBundleRequestFlags Flags, FInstallBundleCompleteDelegate Callback) = 0;
+	// BundleContexts contains all dependencies and has been deduped
+	virtual void RequestUpdateContent(TArrayView<RequestUpdateContentBundleContext> BundleContexts) = 0;
 
 	// Returns true if content is scheduled to be removed the next time the source is initialized
 	// BundleNames contains all dependencies and has been deduped
diff --git a/Engine/Source/Runtime/InstallBundleManager/Public/InstallBundleUtils.h b/Engine/Source/Runtime/InstallBundleManager/Public/InstallBundleUtils.h
--- a/Engine/Source/Runtime/InstallBundleManager/Public/InstallBundleUtils.h
+++ b/Engine/Source/Runtime/InstallBundleManager/Public/InstallBundleUtils.h
@@ -131,4 +131,44 @@ namespace InstallBundleUtil
 	INSTALLBUNDLEMANAGER_API void FinishInstallBundleAsyncIOTasks(TArray<TUniquePtr<FInstallBundleTask>>& Tasks);
 
 	INSTALLBUNDLEMANAGER_API void CleanupInstallBundleAsyncIOTasks(TArray<TUniquePtr<FInstallBundleTask>>& Tasks);
+
+	struct INSTALLBUNDLEMANAGER_API FContentRequestStateStats
+	{
+		double StartTime = 0.0;
+		double EndTime = 0.0;
+		uint64 DataSize = 0;
+		bool bOpen = true;
+
+		double GetElapsedTime() const
+		{
+			return (EndTime > StartTime) ? (EndTime - StartTime) : 0.0;
+		}
+	};
+
+	struct INSTALLBUNDLEMANAGER_API FContentRequestStats
+	{
+		double StartTime = 0.0;
+		double EndTime = 0.0;
+		bool bOpen = true;
+		TMap<FString, FContentRequestStateStats> StateStats;
+
+		double GetElapsedTime() const
+		{
+			return (EndTime > StartTime) ? (EndTime - StartTime) : 0.0;
+		}
+	};
+
+	class INSTALLBUNDLEMANAGER_API FContentRequestStatsMap
+	{
+	private:
+		TMap<FName, InstallBundleUtil::FContentRequestStats> StatsMap;
+
+	public:
+		void StatsBegin(FName BundleName);
+		void StatsEnd(FName BundleName);
+		void StatsBegin(FName BundleName, const TCHAR* State);
+		void StatsEnd(FName BundleName, const TCHAR* State, uint64 DataSize = 0);
+
+		const TMap<FName, InstallBundleUtil::FContentRequestStats>& GetMap() { return StatsMap; }
+	};
 }
