Install Bundle Manger Refactor Snapshot
- Implement background downloading and pausing

diff --git a/Engine/Source/Runtime/InstallBundleManager/Public/InstallBundleSourceInterface.h b/Engine/Source/Runtime/InstallBundleManager/Public/InstallBundleSourceInterface.h
--- a/Engine/Source/Runtime/InstallBundleManager/Public/InstallBundleSourceInterface.h
+++ b/Engine/Source/Runtime/InstallBundleManager/Public/InstallBundleSourceInterface.h
@@ -11,6 +11,7 @@ class IAnalyticsProviderET;
 DECLARE_DELEGATE_ThreeParams(FInstallBundleSourceInitDelegate, TSharedRef<IInstallBundleSource> /*Source*/, EInstallBundleManagerInitResult /*Result*/, bool /*bShouldUseFallbackSource*/)
 
 DECLARE_DELEGATE_OneParam(FInstallBundleCompleteDelegate, FInstallBundleSourceRequestResultInfo /*Result*/);
+DECLARE_DELEGATE_OneParam(FInstallBundlePausedDelegate, FInstallBundleSourcePauseInfo /*PauseInfo*/);
 
 class IInstallBundleSource : public TSharedFromThis<IInstallBundleSource>
 {
@@ -58,7 +59,8 @@ public:
 	{
 		FName BundleName;
 		EInstallBundleRequestFlags Flags = EInstallBundleRequestFlags::None;
-		FInstallBundleCompleteDelegate Callback;
+		FInstallBundlePausedDelegate PausedCallback;
+		FInstallBundleCompleteDelegate CompleteCallback;
 	};
 
 	// Updates content on disk if necessary
diff --git a/Engine/Source/Runtime/InstallBundleManager/Public/InstallBundleTypes.h b/Engine/Source/Runtime/InstallBundleManager/Public/InstallBundleTypes.h
--- a/Engine/Source/Runtime/InstallBundleManager/Public/InstallBundleTypes.h
+++ b/Engine/Source/Runtime/InstallBundleManager/Public/InstallBundleTypes.h
@@ -160,5 +160,7 @@ struct FInstallBundleSourcePauseInfo
 {
 	FName BundleName;
 	EInstallBundlePauseFlags PauseFlags = EInstallBundlePauseFlags::None;
+	// True if the bundle actually transitioned to/from paused,
+	// which is different than the flags changing
 	bool bDidPauseChange = false;
 };
