Install Bundle Manger Refactor Snapshot
- Implement content request handling

diff --git a/Engine/Source/Runtime/InstallBundleManager/Private/InstallBundleTypes.cpp b/Engine/Source/Runtime/InstallBundleManager/Private/InstallBundleTypes.cpp
--- a/Engine/Source/Runtime/InstallBundleManager/Private/InstallBundleTypes.cpp
+++ b/Engine/Source/Runtime/InstallBundleManager/Private/InstallBundleTypes.cpp
@@ -37,19 +37,6 @@ const TCHAR* LexToString(EInstallBundleManagerInitResult Result)
 	return Strings[InstallBundleUtil::CastToUnderlying(Result)];
 }
 
-const TCHAR* LexToString(EBundleState Val)
-{
-	static const TCHAR* Strings[] =
-	{
-		TEXT("NotInstalled"),
-		TEXT("NeedsUpdate"),
-		TEXT("NeedsMount"),
-		TEXT("Mounted"),
-	};
-	static_assert(InstallBundleUtil::CastToUnderlying(EBundleState::Count) == UE_ARRAY_COUNT(Strings), "");
-	return Strings[InstallBundleUtil::CastToUnderlying(Val)];
-}
-
 const TCHAR* LexToString(EInstallBundleContentState State)
 {
 	static const TCHAR* Strings[] =
diff --git a/Engine/Source/Runtime/InstallBundleManager/Private/InstallBundleUtils.cpp b/Engine/Source/Runtime/InstallBundleManager/Private/InstallBundleUtils.cpp
--- a/Engine/Source/Runtime/InstallBundleManager/Private/InstallBundleUtils.cpp
+++ b/Engine/Source/Runtime/InstallBundleManager/Private/InstallBundleUtils.cpp
@@ -19,16 +19,6 @@ namespace InstallBundleUtil
 			&& ConnectionType != ENetworkConnectionType::None;
 	}
 
-	bool StateSignifiesNeedsInstall(EBundleState StateIn)
-	{
-		return (StateIn == EBundleState::NotInstalled || StateIn == EBundleState::NeedsUpdate);
-	}
-
-	bool StateSignifiesNeedsInstall(EInstallBundleContentState StateIn)
-	{
-		return (StateIn == EInstallBundleContentState::NotInstalled || StateIn == EInstallBundleContentState::NeedsUpdate);
-	}
-
 	const TCHAR* GetInstallBundlePauseReason(EInstallBundlePauseFlags Flags)
 	{
 		// Return the most appropriate reason given the flags
diff --git a/Engine/Source/Runtime/InstallBundleManager/Public/InstallBundleSourceInterface.h b/Engine/Source/Runtime/InstallBundleManager/Public/InstallBundleSourceInterface.h
--- a/Engine/Source/Runtime/InstallBundleManager/Public/InstallBundleSourceInterface.h
+++ b/Engine/Source/Runtime/InstallBundleManager/Public/InstallBundleSourceInterface.h
@@ -8,7 +8,7 @@
 class IInstallBundleSource;
 class IAnalyticsProviderET;
 
-DECLARE_DELEGATE_ThreeParams(FInstallBundleSourceInitDelegate, TSharedRef<IInstallBundleSource> /*Source*/, EInstallBundleManagerInitResult /*Result*/, bool /*bShouldUseFallbackSource*/)
+DECLARE_DELEGATE_TwoParams(FInstallBundleSourceInitDelegate, TSharedRef<IInstallBundleSource> /*Source*/, FInstallBundleSourceInitInfo /*InitInfo*/);
 
 DECLARE_DELEGATE_OneParam(FInstallBundleCompleteDelegate, FInstallBundleSourceRequestResultInfo /*Result*/);
 DECLARE_DELEGATE_OneParam(FInstallBundlePausedDelegate, FInstallBundleSourcePauseInfo /*PauseInfo*/);
@@ -44,9 +44,6 @@ public:
 
 	// Whether this source has been initialized or not
 	virtual EInstallBundleManagerInitState GetInitState() const = 0;
-	// Only valid after AsyncInit completes
-	// Returns the result of the last initialization attempt
-	virtual EInstallBundleManagerInitResult GetLastInitResult() const = 0;
 
 	// Returns content version in a "<BuildVersion>-<Platform>" format
 	virtual FString GetContentVersion() const = 0;
@@ -55,17 +52,18 @@ public:
 	// BundleNames contains all dependencies and has been deduped
 	virtual void GetContentState(TArrayView<FName> BundleNames, EInstallBundleGetContentStateFlags Flags, FInstallBundleGetContentStateDelegate Callback) = 0;
 
-	struct RequestUpdateContentBundleContext
+	struct FRequestUpdateContentBundleContext
 	{
 		FName BundleName;
 		EInstallBundleRequestFlags Flags = EInstallBundleRequestFlags::None;
 		FInstallBundlePausedDelegate PausedCallback;
 		FInstallBundleCompleteDelegate CompleteCallback;
+		TSharedPtr<InstallBundleUtil::FContentRequestSharedContext> RequestSharedContext;
 	};
 
 	// Updates content on disk if necessary
 	// BundleContexts contains all dependencies and has been deduped
-	virtual void RequestUpdateContent(TArrayView<RequestUpdateContentBundleContext> BundleContexts) = 0;
+	virtual void RequestUpdateContent(FRequestUpdateContentBundleContext BundleContext) = 0;
 
 	// Returns true if content is scheduled to be removed the next time the source is initialized
 	// BundleNames contains all dependencies and has been deduped
diff --git a/Engine/Source/Runtime/InstallBundleManager/Public/InstallBundleTypes.h b/Engine/Source/Runtime/InstallBundleManager/Public/InstallBundleTypes.h
--- a/Engine/Source/Runtime/InstallBundleManager/Public/InstallBundleTypes.h
+++ b/Engine/Source/Runtime/InstallBundleManager/Public/InstallBundleTypes.h
@@ -39,17 +39,6 @@ enum class EInstallBundleManagerInitResult : int
 };
 INSTALLBUNDLEMANAGER_API const TCHAR* LexToString(EInstallBundleManagerInitResult Result);
 
-// TODO: Needs to be renamed to EInstallBundleState
-enum class EBundleState : int
-{
-	NotInstalled,
-	NeedsUpdate,
-	NeedsMount,
-	Mounted,
-	Count,
-};
-INSTALLBUNDLEMANAGER_API const TCHAR* LexToString(EBundleState Val);
-
 enum class EInstallBundleContentState : int
 {
 	InitializationError,
@@ -127,6 +116,14 @@ struct FInstallBundleRequestInfo
 	TArray<FName> BundlesQueuedForInstall;
 };
 
+struct FInstallBundleSourceInitInfo
+{
+	EInstallBundleManagerInitResult Result = EInstallBundleManagerInitResult::OK;
+	bool bShouldUseFallbackSource = false;
+
+	TMap<FName, bool> BundleUpToDate;
+};
+
 struct FInstallBundleSourceRequestResultInfo
 {
 	FName BundleName;
@@ -138,6 +135,10 @@ struct FInstallBundleSourceRequestResultInfo
 	FString OptionalErrorCode;
 
 	TArray<FString> ContentPaths;
+	// Support platforms that need shaderlibs in the physical FS
+	TSet<FString> NonUFSShaderLibPaths;
+
+	bool bContentWasInstalled = false;
 };
 
 enum class EInstallBundleCancelFlags : int32
diff --git a/Engine/Source/Runtime/InstallBundleManager/Public/InstallBundleUtils.h b/Engine/Source/Runtime/InstallBundleManager/Public/InstallBundleUtils.h
--- a/Engine/Source/Runtime/InstallBundleManager/Public/InstallBundleUtils.h
+++ b/Engine/Source/Runtime/InstallBundleManager/Public/InstallBundleUtils.h
@@ -14,10 +14,6 @@ namespace InstallBundleUtil
 
 	INSTALLBUNDLEMANAGER_API bool HasInternetConnection(ENetworkConnectionType ConnectionType);
 
-	INSTALLBUNDLEMANAGER_API bool StateSignifiesNeedsInstall(EBundleState StateIn);
-
-	INSTALLBUNDLEMANAGER_API bool StateSignifiesNeedsInstall(EInstallBundleContentState StateIn);
-
 	INSTALLBUNDLEMANAGER_API const TCHAR* GetInstallBundlePauseReason(EInstallBundlePauseFlags Flags);
 
 	// It would really be nice to have these in core
@@ -171,4 +167,21 @@ namespace InstallBundleUtil
 
 		const TMap<FName, InstallBundleUtil::FContentRequestStats>& GetMap() { return StatsMap; }
 	};
+
+	struct INSTALLBUNDLEMANAGER_API IBundleSourceContentRequestSharedContext
+	{
+		virtual ~IBundleSourceContentRequestSharedContext() {}
+	};
+
+	struct INSTALLBUNDLEMANAGER_API FContentRequestSharedContext
+	{
+		FContentRequestSharedContext() = default;
+		FContentRequestSharedContext(const FContentRequestSharedContext& Other) = delete;
+		FContentRequestSharedContext(FContentRequestSharedContext&& Other) = default;
+		FContentRequestSharedContext& operator=(const FContentRequestSharedContext& Other) = delete;
+		FContentRequestSharedContext& operator=(FContentRequestSharedContext&& Other) = default;
+
+		TMap<EInstallBundleSourceType, TUniquePtr<IBundleSourceContentRequestSharedContext>> BundleSourceSharedContext;
+	};
+	using FContentRequestSharedContextPtr = TSharedPtr<FContentRequestSharedContext>;
 }
