Undo bad merge path resolve
#rnx

diff --git a/Engine/Source/Runtime/InstallBundleManager/Public/InstallBundleManagerModule.h b/Engine/Source/Runtime/InstallBundleManager/Public/InstallBundleManagerModule.h
new file mode 100644
--- /dev/null
+++ b/Engine/Source/Runtime/InstallBundleManager/Public/InstallBundleManagerModule.h
@@ -0,0 +1,63 @@
+// Copyright 1998-2019 Epic Games, Inc. All Rights Reserved.
+#pragma once
+
+#include "Modules/ModuleInterface.h"
+#include "Modules/ModuleManager.h"
+#include "Misc/ConfigCacheIni.h"
+#include "InstallBundleManagerInterface.h"
+
+/**
+ * Currently empty implementation for InstallBundleModule until things are moved in here.
+ */
+class FInstallBundleManagerModule : public IModuleInterface
+{
+public:
+	// IModuleInterface interface
+	virtual void StartupModule() override
+	{
+	}
+
+	virtual void ShutdownModule() override
+	{
+	}
+};
+
+/**
+ * Base Module Interface for InstallBundleManager implementation modules
+ */
+class IInstallBundleManagerModule : public IModuleInterface
+{
+public:
+	virtual void PreUnloadCallback() override
+	{
+		InstallBundleManager.Reset();
+	}
+
+	IInstallBundleManager* GetInstallBundleManager()
+	{
+		return InstallBundleManager.Get();
+	}
+
+protected:
+	TUniquePtr<IInstallBundleManager> InstallBundleManager;
+};
+
+/**
+ * Module Interface for InstallBundleManager implementation modules
+ */
+template<class InstallBundleManagerModuleImpl>
+class TInstallBundleManagerModule : public IInstallBundleManagerModule
+{
+public:
+	virtual void StartupModule() override
+	{
+		// Only instantiate the bundle manager if this is the version the game has been configured to use
+		FString ModuleName;
+		GConfig->GetString(TEXT("InstallBundleManager"), TEXT("ModuleName"), ModuleName, GEngineIni);
+
+		if (FModuleManager::Get().GetModule(*ModuleName) == this)
+		{
+			InstallBundleManager = MakeUnique<InstallBundleManagerModuleImpl>();
+		}
+	}
+};
