Install Bundle Manager Refactor
- Add AdditionalRootDirs to bundle source results

diff --git a/Engine/Source/Runtime/InstallBundleManager/Public/InstallBundleTypes.h b/Engine/Source/Runtime/InstallBundleManager/Public/InstallBundleTypes.h
--- a/Engine/Source/Runtime/InstallBundleManager/Public/InstallBundleTypes.h
+++ b/Engine/Source/Runtime/InstallBundleManager/Public/InstallBundleTypes.h
@@ -135,6 +135,7 @@ struct FInstallBundleSourceRequestResultInfo
 	FString OptionalErrorCode;
 
 	TArray<FString> ContentPaths;
+	TArray<FString> AdditionalRootDirs;
 	// Support platforms that need shaderlibs in the physical FS
 	TSet<FString> NonUFSShaderLibPaths;
 
