Install bundle manager refactor
- Replace HasBuildMetaData() with HasBundleSource().  A source won't exist in the case it falls back to another source so this is a drop in replacement if there is no build metadata.

diff --git a/Engine/Source/Runtime/InstallBundleManager/Public/InstallBundleManagerInterface.h b/Engine/Source/Runtime/InstallBundleManager/Public/InstallBundleManagerInterface.h
--- a/Engine/Source/Runtime/InstallBundleManager/Public/InstallBundleManagerInterface.h
+++ b/Engine/Source/Runtime/InstallBundleManager/Public/InstallBundleManagerInterface.h
@@ -122,8 +122,7 @@ public:
 
 	virtual ~IInstallBundleManager() {}
 
-	// TODO: should query if we are using the BPT source
-	virtual bool HasBuildMetaData() const = 0;
+	virtual bool HasBundleSource(EInstallBundleSourceType SourceType) const = 0;
 
 	virtual FDelegateHandle PushInitErrorCallback(FInstallBundleManagerInitErrorHandler Callback) = 0;
 	virtual void PopInitErrorCallback() = 0;
diff --git a/Engine/Source/Runtime/NullInstallBundleManager/Source/NullInstallBundleManager.cpp b/Engine/Source/Runtime/NullInstallBundleManager/Source/NullInstallBundleManager.cpp
--- a/Engine/Source/Runtime/NullInstallBundleManager/Source/NullInstallBundleManager.cpp
+++ b/Engine/Source/Runtime/NullInstallBundleManager/Source/NullInstallBundleManager.cpp
@@ -6,7 +6,7 @@
 
 class FNullInstallBundleManager : public IInstallBundleManager
 {
-	virtual bool HasBuildMetaData() const override
+	virtual bool HasBundleSource(EInstallBundleSourceType SourceType) const override
 	{
 		return false;
 	}
