Install bundle manager refactor
- Refactor Progress
- Add logging for bundle versions
- Fix min/max bundle version logging not working when some bundles had no version

diff --git a/Engine/Source/Runtime/InstallBundleManager/Private/BundlePrereqCombinedStatusHelper.cpp b/Engine/Source/Runtime/InstallBundleManager/Private/BundlePrereqCombinedStatusHelper.cpp
--- a/Engine/Source/Runtime/InstallBundleManager/Private/BundlePrereqCombinedStatusHelper.cpp
+++ b/Engine/Source/Runtime/InstallBundleManager/Private/BundlePrereqCombinedStatusHelper.cpp
@@ -4,18 +4,7 @@
 #include "Containers/Ticker.h"
 #include "InstallBundleManagerPrivatePCH.h"
 
-//PRAGMA_DISABLE_OPTIMIZATION
-
 FBundlePrereqCombinedStatusHelper::FBundlePrereqCombinedStatusHelper()
-: DownloadWeight(1.f)
-, InstallWeight(1.f)
-, RequiredBundleNames()
-, BundleStatusCache()
-, CachedBundleWeights()
-, CurrentCombinedStatus()
-, bBundleNeedsUpdate(false)
-, InstallBundleManager(nullptr)
-, TickHandle()
 {
 	SetupDelegates();
 }
@@ -40,8 +29,6 @@ FBundlePrereqCombinedStatusHelper& FBundlePrereqCombinedStatusHelper::operator=(
 	if (this != &Other)
 	{
 		//Just copy all this data
-		DownloadWeight = Other.DownloadWeight;
-		InstallWeight = Other.InstallWeight;
 		RequiredBundleNames = Other.RequiredBundleNames;
 		BundleStatusCache = Other.BundleStatusCache;
 		CachedBundleWeights = Other.CachedBundleWeights;
@@ -61,8 +48,6 @@ FBundlePrereqCombinedStatusHelper& FBundlePrereqCombinedStatusHelper::operator=(
 	if (this != &Other)
 	{
 		//Just copy small data
-		DownloadWeight = Other.DownloadWeight;
-		InstallWeight = Other.InstallWeight;
 		CurrentCombinedStatus = Other.CurrentCombinedStatus;
 		bBundleNeedsUpdate = Other.bBundleNeedsUpdate;
 		InstallBundleManager = Other.InstallBundleManager;
@@ -191,7 +176,7 @@ void FBundlePrereqCombinedStatusHelper::UpdateCombinedStatus()
 	if ((BundleStatusCache.Num() < RequiredBundleNames.Num())
 		&& (BundleStatusCache.Num() > 0))
 	{
-		EarliestBundleState = EInstallBundleStatus::Downloading;
+		EarliestBundleState = EInstallBundleStatus::Installing;
 	}
 	
 	float EarliestFinishingPercent = 1.0f;
@@ -272,7 +257,7 @@ float FBundlePrereqCombinedStatusHelper::GetCombinedProgressPercent() const
 		const float* FoundWeight = CachedBundleWeights.Find(BundleStatusPair.Key);
 		if (ensureAlwaysMsgf((nullptr != FoundWeight), TEXT("Found missing entry for BundleWeight! Bundle %s does not have a weight entry!"), *(BundleStatusPair.Key.ToString())))
 		{
-			AllBundleProgressPercent += ((*FoundWeight) * GetIndividualWeightedProgressPercent(BundleStatusPair.Value));
+			AllBundleProgressPercent += ((*FoundWeight) * BundleStatusPair.Value.Install_Percent);
 		}
 	}
 	
@@ -280,58 +265,6 @@ float FBundlePrereqCombinedStatusHelper::GetCombinedProgressPercent() const
 	return AllBundleProgressPercent;
 }
 
-float FBundlePrereqCombinedStatusHelper::GetIndividualWeightedProgressPercent(const FInstallBundleStatus& BundleStatus) const
-{
-	const float TotalWeight = DownloadWeight + InstallWeight;
-	float CombinedOverallProgressPercent = 0.f;
-	
-	//If we are done installing then lets just coung our progress as 1.0f. Finishing progress is handled in UpdateCombinedStatus
-	if (BundleStatus.Status > EInstallBundleStatus::Installing)
-	{
-		CombinedOverallProgressPercent = 1.0f;
-	}
-	else
-	{
-		if (ensureAlwaysMsgf((TotalWeight > 0), TEXT("Invalid Weights for FBundleStatusTracker! Need to have weights > 0.f!")))
-		{
-			float DownloadWeightedPercent = 0.f;
-			bool bDidHaveBackgroundDownloadProgress = false;
-			
-			if (DownloadWeight > 0.f)
-			{
-				//Skip background downloads from contributing to progress if they aren't in use
-				bDidHaveBackgroundDownloadProgress = (BundleStatus.BackgroundDownloadProgress.IsSet() && (BundleStatus.BackgroundDownloadProgress->BytesDownloaded > 0));
-				if (bDidHaveBackgroundDownloadProgress)
-				{
-					const float AppliedDownloadWeight = DownloadWeight / TotalWeight;
-					
-					if (BundleStatus.Status > EInstallBundleStatus::Downloading)
-					{
-						//if we have moved past the download phase, just consider these
-						//downloads as having finished
-						DownloadWeightedPercent = 1.0f * AppliedDownloadWeight;
-					}
-					else
-					{
-						DownloadWeightedPercent = BundleStatus.BackgroundDownloadProgress->PercentComplete * AppliedDownloadWeight;
-					}
-				}
-			}
-			
-			float InstallWeightedPercent = 0.f;
-			
-			//If we didn't have background download progress, just use our InstallWeight as we don't have Download Progress to use
-			const float AppliedInstallWeight = bDidHaveBackgroundDownloadProgress ? (InstallWeight / TotalWeight) : 1.0f;
-			InstallWeightedPercent = BundleStatus.Install_Percent * AppliedInstallWeight;
-			
-			CombinedOverallProgressPercent = DownloadWeightedPercent + InstallWeightedPercent;
-			FMath::Clamp(CombinedOverallProgressPercent, 0.f, 1.0f);
-		}
-	}
-	
-	return CombinedOverallProgressPercent;
-}
-
 bool FBundlePrereqCombinedStatusHelper::Tick(float dt)
 {
 	UpdateBundleCache();
@@ -377,5 +310,3 @@ void FBundlePrereqCombinedStatusHelper::OnBundleInstallPauseChanged(FInstallBund
 		BundleCacheInfo.PauseFlags = PauseInfo.PauseFlags;
 	}
 }
-
-//PRAGMA_ENABLE_OPTIMIZATION
diff --git a/Engine/Source/Runtime/InstallBundleManager/Private/InstallBundleTypes.cpp b/Engine/Source/Runtime/InstallBundleManager/Private/InstallBundleTypes.cpp
--- a/Engine/Source/Runtime/InstallBundleManager/Private/InstallBundleTypes.cpp
+++ b/Engine/Source/Runtime/InstallBundleManager/Private/InstallBundleTypes.cpp
@@ -12,6 +12,7 @@ const TCHAR* LexToString(EInstallBundleSourceType Type)
 		TEXT("Bulk"),
 		TEXT("BuildPatchServices"),
 	};
+
 	static_assert(InstallBundleUtil::CastToUnderlying(EInstallBundleSourceType::Count) == UE_ARRAY_COUNT(Strings), "");
 	return Strings[InstallBundleUtil::CastToUnderlying(Type)];
 }
@@ -33,6 +34,7 @@ const TCHAR* LexToString(EInstallBundleManagerInitResult Result)
 		TEXT("NoInternetConnectionError"),
 		TEXT("ConfigurationError"),
 	};
+
 	static_assert(InstallBundleUtil::CastToUnderlying(EInstallBundleManagerInitResult::Count) == UE_ARRAY_COUNT(Strings), "");
 	return Strings[InstallBundleUtil::CastToUnderlying(Result)];
 }
@@ -46,8 +48,8 @@ const TCHAR* LexToString(EInstallBundleContentState State)
 		TEXT("NeedsUpdate"),
 		TEXT("UpToDate"),
 	};
-	static_assert(InstallBundleUtil::CastToUnderlying(EInstallBundleContentState::Count) == UE_ARRAY_COUNT(Strings), "");
 
+	static_assert(InstallBundleUtil::CastToUnderlying(EInstallBundleContentState::Count) == UE_ARRAY_COUNT(Strings), "");
 	return Strings[InstallBundleUtil::CastToUnderlying(State)];
 }
 
@@ -63,7 +65,21 @@ const TCHAR* LexToString(EInstallBundleResult Result)
 		TEXT("UserCancelledError"),
 		TEXT("InitializationError"),
 	};
-	static_assert(InstallBundleUtil::CastToUnderlying(EInstallBundleResult::Count) == UE_ARRAY_COUNT(Strings), "");
 
+	static_assert(InstallBundleUtil::CastToUnderlying(EInstallBundleResult::Count) == UE_ARRAY_COUNT(Strings), "");
 	return Strings[InstallBundleUtil::CastToUnderlying(Result)];
 }
+
+const TCHAR* LexToString(EInstallBundleStatus Status)
+{
+	static const TCHAR* Strings[] =
+	{
+		TEXT("Requested"),
+		TEXT("Installing"),
+		TEXT("Finishing"),
+		TEXT("Installed"),
+	};
+
+	static_assert(InstallBundleUtil::CastToUnderlying(EInstallBundleStatus::Count) == UE_ARRAY_COUNT(Strings), "");
+	return Strings[InstallBundleUtil::CastToUnderlying(Status)];
+}
diff --git a/Engine/Source/Runtime/InstallBundleManager/Public/BundlePrereqCombinedStatusHelper.h b/Engine/Source/Runtime/InstallBundleManager/Public/BundlePrereqCombinedStatusHelper.h
--- a/Engine/Source/Runtime/InstallBundleManager/Public/BundlePrereqCombinedStatusHelper.h
+++ b/Engine/Source/Runtime/InstallBundleManager/Public/BundlePrereqCombinedStatusHelper.h
@@ -47,11 +47,6 @@ public:
 	//Get current CombinedBundleStatus for everything setup to track
 	const FCombinedBundleStatus& GetCurrentCombinedState() const;
 	
-	//How to weight downloads vs. installs. Defaults to even. Does not have to add up to 1.0.
-	//Setting Download to .5 and Install to .5 will be the same as setting Download to 1.f and Install to 1.f.
-	float DownloadWeight;
-	float InstallWeight;
-	
 private:
 	bool Tick(float dt);
 	void UpdateBundleCache();
@@ -65,7 +60,6 @@ private:
 	void OnBundleInstallPauseChanged(FInstallBundlePauseInfo PauseInfo);
 	
 	float GetCombinedProgressPercent() const;
-	float GetIndividualWeightedProgressPercent(const FInstallBundleStatus& Bundle) const;
 	
 private:
 	//All bundles we need including pre-reqs
@@ -79,8 +73,8 @@ private:
 	
 	FCombinedBundleStatus CurrentCombinedStatus;
 	
-	bool bBundleNeedsUpdate;
+	bool bBundleNeedsUpdate = false;
 	
-	IInstallBundleManager* InstallBundleManager;
+	IInstallBundleManager* InstallBundleManager = nullptr;
 	FDelegateHandle TickHandle;
 };
diff --git a/Engine/Source/Runtime/InstallBundleManager/Public/InstallBundleManagerInterface.h b/Engine/Source/Runtime/InstallBundleManager/Public/InstallBundleManagerInterface.h
--- a/Engine/Source/Runtime/InstallBundleManager/Public/InstallBundleManagerInterface.h
+++ b/Engine/Source/Runtime/InstallBundleManager/Public/InstallBundleManagerInterface.h
@@ -10,72 +10,14 @@
 
 class IAnalyticsProviderET;
 
-enum class EInstallBundleStatus : int
-{
-	Requested,
-	QueuedForDownload,
-	Downloading,
-	QueuedForInstall,
-	Installing,
-	QueuedForFinish,
-	Finishing,
-	Installed,
-	Count,
-};
-
-inline const TCHAR* LexToString(EInstallBundleStatus Status)
-{
-	using UnderType = __underlying_type(EInstallBundleStatus);
-	static const TCHAR* Strings[] =
-	{
-		TEXT("Requested"),
-		TEXT("QueuedForDownload"),
-		TEXT("Downloading"),
-		TEXT("QueuedForInstall"),
-		TEXT("Installing"),
-		TEXT("QueuedForFinish"),
-		TEXT("Finishing"),
-		TEXT("Installed"),
-	};
-
-	static_assert(static_cast<UnderType>(EInstallBundleStatus::Count) == UE_ARRAY_COUNT(Strings), "");
-
-	return Strings[static_cast<UnderType>(Status)];
-}
-
-struct FInstallBundleDownloadProgress
-{
-	// Num bytes received
-	uint64 BytesDownloaded = 0;
-	// Num bytes written to storage (<= BytesDownloaded)
-	uint64 BytesDownloadedAndWritten = 0;
-	// Num bytes needed
-	uint64 TotalBytesToDownload = 0;
-	// Num bytes that failed to download
-	uint64 TotalBytesFailedToDownload = 0;
-	float PercentComplete = 0;
-};
-
 struct FInstallBundleStatus
 {
 	FName BundleName;
 
-	EInstallBundleStatus Status = EInstallBundleStatus::QueuedForDownload;
+	EInstallBundleStatus Status = EInstallBundleStatus::Requested;
 
 	EInstallBundlePauseFlags PauseFlags = EInstallBundlePauseFlags::None;
 
-	FText StatusText;
-
-	// Download progress of EInstallBundleStatus::Downloading
-	// Will be set if Status >= EInstallBundleStatus::Downloading
-	TOptional<FInstallBundleDownloadProgress> BackgroundDownloadProgress;
-
-	// Download progress of EInstallBundleStatus::Install
-	// Will be set if Status >= EInstallBundleStatus::Install
-	// We may download during install if background downloads are turned off or fail
-	// We may also do small downloads during install as a normal part of installation
-	TOptional<FInstallBundleDownloadProgress> InstallDownloadProgress;
-
 	float Install_Percent = 0;
 
 	float Finishing_Percent = 0;
diff --git a/Engine/Source/Runtime/InstallBundleManager/Public/InstallBundleSourceInterface.h b/Engine/Source/Runtime/InstallBundleManager/Public/InstallBundleSourceInterface.h
--- a/Engine/Source/Runtime/InstallBundleManager/Public/InstallBundleSourceInterface.h
+++ b/Engine/Source/Runtime/InstallBundleManager/Public/InstallBundleSourceInterface.h
@@ -75,7 +75,8 @@ public:
 	virtual EInstallBundleRequestFlags GetModifyableContentRequestFlags() const { return EInstallBundleRequestFlags::None; }
 	virtual void UpdateContentRequestFlags(TArrayView<const FName> BundleNames, EInstallBundleRequestFlags AddFlags, EInstallBundleRequestFlags RemoveFlags) {}
 
-	// TODO: GetBundleProgress() - Not sure how this will work yet
+	// Derived classes should implement this if their content install will take a significant amount of time
+	virtual TOptional<FInstallBundleSourceStatus> GetBundleProgress(FName BundleName) const { return TOptional<FInstallBundleSourceStatus>(); }
 
 	// Called by bundle manager to pass through command line options to simulate errors
 	virtual void SetErrorSimulationCommands(const FString& CommandLine) {}
diff --git a/Engine/Source/Runtime/InstallBundleManager/Public/InstallBundleTypes.h b/Engine/Source/Runtime/InstallBundleManager/Public/InstallBundleTypes.h
--- a/Engine/Source/Runtime/InstallBundleManager/Public/InstallBundleTypes.h
+++ b/Engine/Source/Runtime/InstallBundleManager/Public/InstallBundleTypes.h
@@ -166,3 +166,20 @@ struct FInstallBundleSourcePauseInfo
 	// which is different than the flags changing
 	bool bDidPauseChange = false;
 };
+
+enum class EInstallBundleStatus : int
+{
+	Requested,
+	Installing,
+	Finishing,
+	Installed,
+	Count,
+};
+INSTALLBUNDLEMANAGER_API const TCHAR* LexToString(EInstallBundleStatus Status);
+
+struct FInstallBundleSourceStatus
+{
+	FName BundleName;
+
+	float Install_Percent = 0;
+};
