Install Bundle Refactor
- Type renames

diff --git a/Engine/Source/Runtime/InstallBundleManager/Private/BundlePrereqCombinedStatusHelper.cpp b/Engine/Source/Runtime/InstallBundleManager/Private/BundlePrereqCombinedStatusHelper.cpp
--- a/Engine/Source/Runtime/InstallBundleManager/Private/BundlePrereqCombinedStatusHelper.cpp
+++ b/Engine/Source/Runtime/InstallBundleManager/Private/BundlePrereqCombinedStatusHelper.cpp
@@ -146,7 +146,7 @@ void FBundlePrereqCombinedStatusHelper::UpdateBundleCache()
 	{
 		for (FName& BundleName : RequiredBundleNames)
 		{
-			TOptional<FInstallBundleStatus> BundleProgress = InstallBundleManager->GetBundleProgress(BundleName);
+			TOptional<FInstallBundleProgress> BundleProgress = InstallBundleManager->GetBundleProgress(BundleName);
 			
 			//Copy progress to the cache as long as we have progress to copy.
 			if (BundleProgress.IsSet())
@@ -176,11 +176,11 @@ void FBundlePrereqCombinedStatusHelper::UpdateCombinedStatus()
 	if ((BundleStatusCache.Num() < RequiredBundleNames.Num())
 		&& (BundleStatusCache.Num() > 0))
 	{
-		EarliestBundleState = EInstallBundleStatus::Installing;
+		EarliestBundleState = EInstallBundleStatus::Updating;
 	}
 	
 	float EarliestFinishingPercent = 1.0f;
-	for (const TPair<FName,FInstallBundleStatus>& BundlePair : BundleStatusCache)
+	for (const TPair<FName,FInstallBundleProgress>& BundlePair : BundleStatusCache)
 	{
 		if (BundlePair.Value.Status < EarliestBundleState)
 		{
@@ -199,7 +199,7 @@ void FBundlePrereqCombinedStatusHelper::UpdateCombinedStatus()
 	
 	//if we have any paused bundles, and we have any bundle that isn't finished installed, we are Paused
 	//if everything is installed ignore the pause flags as we completed after pausing the bundles
-	CurrentCombinedStatus.bIsPaused = (bIsAnythingPaused && (EarliestBundleState < EInstallBundleStatus::Installed));
+	CurrentCombinedStatus.bIsPaused = (bIsAnythingPaused && (EarliestBundleState < EInstallBundleStatus::Ready));
 	if(CurrentCombinedStatus.bIsPaused)
 	{
 		CurrentCombinedStatus.CombinedPauseFlags = CombinedPauseFlags;
@@ -217,7 +217,7 @@ void FBundlePrereqCombinedStatusHelper::UpdateCombinedStatus()
 	{
 		CurrentCombinedStatus.CombinedState = FCombinedBundleStatus::ECombinedBundleStateEnum::Initializing;
 	}
-	else if (EarliestBundleState <= EInstallBundleStatus::Installing)
+	else if (EarliestBundleState <= EInstallBundleStatus::Updating)
 	{
 		CurrentCombinedStatus.CombinedState = FCombinedBundleStatus::ECombinedBundleStateEnum::Updating;
 	}
@@ -235,7 +235,7 @@ void FBundlePrereqCombinedStatusHelper::UpdateCombinedStatus()
 			CurrentCombinedStatus.CombinedState = FCombinedBundleStatus::ECombinedBundleStateEnum::Updating;
 		}
 	}
-	else if (EarliestBundleState == EInstallBundleStatus::Installed)
+	else if (EarliestBundleState == EInstallBundleStatus::Ready)
 	{
 		CurrentCombinedStatus.CombinedState = FCombinedBundleStatus::ECombinedBundleStateEnum::Finished;
 		CurrentCombinedStatus.bDoesCurrentStateSupportPausing = false;
@@ -252,7 +252,7 @@ float FBundlePrereqCombinedStatusHelper::GetCombinedProgressPercent() const
 	
 	ensureAlwaysMsgf((CachedBundleWeights.Num() >= BundleStatusCache.Num()), TEXT("Missing Cache Entries for BundleWeights!Any missing bundles will have 0 for their progress!"));
 	
-	for (const TPair<FName,FInstallBundleStatus>& BundleStatusPair : BundleStatusCache)
+	for (const TPair<FName,FInstallBundleProgress>& BundleStatusPair : BundleStatusCache)
 	{
 		const float* FoundWeight = CachedBundleWeights.Find(BundleStatusPair.Key);
 		if (ensureAlwaysMsgf((nullptr != FoundWeight), TEXT("Found missing entry for BundleWeight! Bundle %s does not have a weight entry!"), *(BundleStatusPair.Key.ToString())))
@@ -288,10 +288,10 @@ void FBundlePrereqCombinedStatusHelper::OnBundleInstallComplete(FInstallBundleRe
 	if (bBundleCompletedSuccessfully && bWasRequiredBundle)
 	{
 		//Make sure our BundleCache shows this as finished all the way
-		FInstallBundleStatus& BundleCacheInfo = BundleStatusCache.FindOrAdd(CompletedBundleName);
-		BundleCacheInfo.Status = EInstallBundleStatus::Installed;
+		FInstallBundleProgress& BundleCacheInfo = BundleStatusCache.FindOrAdd(CompletedBundleName);
+		BundleCacheInfo.Status = EInstallBundleStatus::Ready;
 		
-		TOptional<FInstallBundleStatus> BundleProgress = InstallBundleManager->GetBundleProgress(CompletedBundleName);
+		TOptional<FInstallBundleProgress> BundleProgress = InstallBundleManager->GetBundleProgress(CompletedBundleName);
 		if (ensureAlwaysMsgf(BundleProgress.IsSet(), TEXT("Expected to find BundleProgress for completed bundle, but did not. Leaving old progress values")))
 		{
 			BundleCacheInfo = BundleProgress.GetValue();
@@ -306,7 +306,7 @@ void FBundlePrereqCombinedStatusHelper::OnBundleInstallPauseChanged(FInstallBund
 	const bool bWasRequiredBundle = RequiredBundleNames.Contains(PauseInfo.BundleName);
 	if (bWasRequiredBundle)
 	{
-		FInstallBundleStatus& BundleCacheInfo = BundleStatusCache.FindOrAdd(PauseInfo.BundleName);
+		FInstallBundleProgress& BundleCacheInfo = BundleStatusCache.FindOrAdd(PauseInfo.BundleName);
 		BundleCacheInfo.PauseFlags = PauseInfo.PauseFlags;
 	}
 }
diff --git a/Engine/Source/Runtime/InstallBundleManager/Private/InstallBundleTypes.cpp b/Engine/Source/Runtime/InstallBundleManager/Private/InstallBundleTypes.cpp
--- a/Engine/Source/Runtime/InstallBundleManager/Private/InstallBundleTypes.cpp
+++ b/Engine/Source/Runtime/InstallBundleManager/Private/InstallBundleTypes.cpp
@@ -75,9 +75,9 @@ const TCHAR* LexToString(EInstallBundleStatus Status)
 	static const TCHAR* Strings[] =
 	{
 		TEXT("Requested"),
-		TEXT("Installing"),
+		TEXT("Updating"),
 		TEXT("Finishing"),
-		TEXT("Installed"),
+		TEXT("Ready"),
 	};
 
 	static_assert(InstallBundleUtil::CastToUnderlying(EInstallBundleStatus::Count) == UE_ARRAY_COUNT(Strings), "");
diff --git a/Engine/Source/Runtime/InstallBundleManager/Public/BundlePrereqCombinedStatusHelper.h b/Engine/Source/Runtime/InstallBundleManager/Public/BundlePrereqCombinedStatusHelper.h
--- a/Engine/Source/Runtime/InstallBundleManager/Public/BundlePrereqCombinedStatusHelper.h
+++ b/Engine/Source/Runtime/InstallBundleManager/Public/BundlePrereqCombinedStatusHelper.h
@@ -66,7 +66,7 @@ private:
 	TArray<FName> RequiredBundleNames;
 	
 	//Internal Cache of all bundle statuses to track progress
-	TMap<FName, FInstallBundleStatus> BundleStatusCache;
+	TMap<FName, FInstallBundleProgress> BundleStatusCache;
 	
 	//Bundle weights that determine what % of the overall install each bundle represents
 	TMap<FName, float> CachedBundleWeights;
diff --git a/Engine/Source/Runtime/InstallBundleManager/Public/InstallBundleManagerInterface.h b/Engine/Source/Runtime/InstallBundleManager/Public/InstallBundleManagerInterface.h
--- a/Engine/Source/Runtime/InstallBundleManager/Public/InstallBundleManagerInterface.h
+++ b/Engine/Source/Runtime/InstallBundleManager/Public/InstallBundleManagerInterface.h
@@ -10,7 +10,7 @@
 
 class IAnalyticsProviderET;
 
-struct FInstallBundleStatus
+struct FInstallBundleProgress
 {
 	FName BundleName;
 
@@ -98,7 +98,7 @@ public:
 
 	virtual void RequestPausedBundleCallback() = 0;
 
-	virtual TOptional<FInstallBundleStatus> GetBundleProgress(FName BundleName) const = 0;
+	virtual TOptional<FInstallBundleProgress> GetBundleProgress(FName BundleName) const = 0;
 
 	virtual EInstallBundleRequestFlags GetModifyableContentRequestFlags() const = 0;
 	virtual void UpdateContentRequestFlags(FName BundleName, EInstallBundleRequestFlags AddFlags, EInstallBundleRequestFlags RemoveFlags) = 0;
diff --git a/Engine/Source/Runtime/InstallBundleManager/Public/InstallBundleSourceInterface.h b/Engine/Source/Runtime/InstallBundleManager/Public/InstallBundleSourceInterface.h
--- a/Engine/Source/Runtime/InstallBundleManager/Public/InstallBundleSourceInterface.h
+++ b/Engine/Source/Runtime/InstallBundleManager/Public/InstallBundleSourceInterface.h
@@ -76,7 +76,7 @@ public:
 	virtual void UpdateContentRequestFlags(TArrayView<const FName> BundleNames, EInstallBundleRequestFlags AddFlags, EInstallBundleRequestFlags RemoveFlags) {}
 
 	// Derived classes should implement this if their content install will take a significant amount of time
-	virtual TOptional<FInstallBundleSourceStatus> GetBundleProgress(FName BundleName) const { return TOptional<FInstallBundleSourceStatus>(); }
+	virtual TOptional<FInstallBundleSourceProgress> GetBundleProgress(FName BundleName) const { return TOptional<FInstallBundleSourceProgress>(); }
 
 	// Called by bundle manager to pass through command line options to simulate errors
 	virtual void SetErrorSimulationCommands(const FString& CommandLine) {}
diff --git a/Engine/Source/Runtime/InstallBundleManager/Public/InstallBundleTypes.h b/Engine/Source/Runtime/InstallBundleManager/Public/InstallBundleTypes.h
--- a/Engine/Source/Runtime/InstallBundleManager/Public/InstallBundleTypes.h
+++ b/Engine/Source/Runtime/InstallBundleManager/Public/InstallBundleTypes.h
@@ -170,14 +170,14 @@ struct FInstallBundleSourcePauseInfo
 enum class EInstallBundleStatus : int
 {
 	Requested,
-	Installing,
+	Updating,
 	Finishing,
-	Installed,
+	Ready,
 	Count,
 };
 INSTALLBUNDLEMANAGER_API const TCHAR* LexToString(EInstallBundleStatus Status);
 
-struct FInstallBundleSourceStatus
+struct FInstallBundleSourceProgress
 {
 	FName BundleName;
 
diff --git a/Engine/Source/Runtime/NullInstallBundleManager/Source/NullInstallBundleManager.cpp b/Engine/Source/Runtime/NullInstallBundleManager/Source/NullInstallBundleManager.cpp
--- a/Engine/Source/Runtime/NullInstallBundleManager/Source/NullInstallBundleManager.cpp
+++ b/Engine/Source/Runtime/NullInstallBundleManager/Source/NullInstallBundleManager.cpp
@@ -108,9 +108,9 @@ class FNullInstallBundleManager : public IInstallBundleManager
 
 	}
 
-	virtual TOptional<FInstallBundleStatus> GetBundleProgress(FName BundleName) const override
+	virtual TOptional<FInstallBundleProgress> GetBundleProgress(FName BundleName) const override
 	{
-		return TOptional<FInstallBundleStatus>();
+		return TOptional<FInstallBundleProgress>();
 	}
 
 	virtual EInstallBundleRequestFlags GetModifyableContentRequestFlags() const override
