Install Bundle Manager Refactor
- Cleanup Interface
- Fix some failing asserts
- Fix engine sleeping during PSO caching
- Enable New Bundle Manager

diff --git a/Engine/Source/Runtime/InstallBundleManager/Private/InstallBundleManagerInterface.cpp b/Engine/Source/Runtime/InstallBundleManager/Private/InstallBundleManagerInterface.cpp
--- a/Engine/Source/Runtime/InstallBundleManager/Private/InstallBundleManagerInterface.cpp
+++ b/Engine/Source/Runtime/InstallBundleManager/Private/InstallBundleManagerInterface.cpp
@@ -39,3 +39,43 @@ IInstallBundleManager* IInstallBundleManager::GetPlatformInstallBundleManager()
 	return Manager;
 }
 
+FInstallBundleRequestInfo IInstallBundleManager::RequestUpdateContent(FName BundleName, EInstallBundleRequestFlags Flags)
+{
+	return RequestUpdateContent(MakeArrayView(&BundleName, 1), Flags);
+}
+
+void IInstallBundleManager::GetContentState(FName BundleName, EInstallBundleGetContentStateFlags Flags, bool bAddDependencies, FInstallBundleGetContentStateDelegate Callback, FName RequestTag /*= TEXT("None")*/)
+{
+	return GetContentState(MakeArrayView(&BundleName, 1), Flags, bAddDependencies, MoveTemp(Callback), RequestTag);
+}
+
+void IInstallBundleManager::RequestRemoveContentOnNextInit(FName RemoveName, TArrayView<const FName> KeepNames /*= TArrayView<const FName>()*/)
+{
+	RequestRemoveContentOnNextInit(MakeArrayView(&RemoveName, 1), KeepNames);
+}
+
+void IInstallBundleManager::CancelRequestRemoveContentOnNextInit(FName BundleName)
+{
+	CancelRequestRemoveContentOnNextInit(MakeArrayView(&BundleName, 1));
+}
+
+void IInstallBundleManager::CancelUpdateContent(FName BundleName, EInstallBundleCancelFlags Flags)
+{
+	CancelUpdateContent(MakeArrayView(&BundleName, 1), Flags);
+}
+
+void IInstallBundleManager::PauseUpdateContent(FName BundleName)
+{
+	PauseUpdateContent(MakeArrayView(&BundleName, 1));
+}
+
+void IInstallBundleManager::ResumeUpdateContent(FName BundleName)
+{
+	ResumeUpdateContent(MakeArrayView(&BundleName, 1));
+}
+
+void IInstallBundleManager::UpdateContentRequestFlags(FName BundleName, EInstallBundleRequestFlags AddFlags, EInstallBundleRequestFlags RemoveFlags)
+{
+	UpdateContentRequestFlags(MakeArrayView(&BundleName, 1), AddFlags, RemoveFlags);
+}
+
diff --git a/Engine/Source/Runtime/InstallBundleManager/Public/InstallBundleManagerInterface.h b/Engine/Source/Runtime/InstallBundleManager/Public/InstallBundleManagerInterface.h
--- a/Engine/Source/Runtime/InstallBundleManager/Public/InstallBundleManagerInterface.h
+++ b/Engine/Source/Runtime/InstallBundleManager/Public/InstallBundleManagerInterface.h
@@ -73,27 +73,27 @@ public:
 
 	virtual EInstallBundleManagerInitState GetInitState() const = 0;
 
-	virtual FInstallBundleRequestInfo RequestUpdateContent(FName BundleName, EInstallBundleRequestFlags Flags) = 0;
+	FInstallBundleRequestInfo RequestUpdateContent(FName BundleName, EInstallBundleRequestFlags Flags);
 	virtual FInstallBundleRequestInfo RequestUpdateContent(TArrayView<const FName> BundleNames, EInstallBundleRequestFlags Flags) = 0;
 
-	virtual void GetContentState(FName BundleName, EInstallBundleGetContentStateFlags Flags, bool bAddDependencies, FInstallBundleGetContentStateDelegate Callback, FName RequestTag = TEXT("None")) = 0;
+	void GetContentState(FName BundleName, EInstallBundleGetContentStateFlags Flags, bool bAddDependencies, FInstallBundleGetContentStateDelegate Callback, FName RequestTag = TEXT("None"));
 	virtual void GetContentState(TArrayView<const FName> BundleNames, EInstallBundleGetContentStateFlags Flags, bool bAddDependencies, FInstallBundleGetContentStateDelegate Callback, FName RequestTag = TEXT("None")) = 0;
 
     virtual void CancelAllGetContentStateRequestsForTag(FName RequestTag) = 0;
     
-	virtual void RequestRemoveContentOnNextInit(FName RemoveName, TArrayView<const FName> KeepNames = TArrayView<const FName>()) = 0;
+	void RequestRemoveContentOnNextInit(FName RemoveName, TArrayView<const FName> KeepNames = TArrayView<const FName>());
 	virtual void RequestRemoveContentOnNextInit(TArrayView<const FName> RemoveNames, TArrayView<const FName> KeepNames = TArrayView<const FName>()) = 0;
 
-	virtual void CancelRequestRemoveContentOnNextInit(FName BundleName) = 0;
+	void CancelRequestRemoveContentOnNextInit(FName BundleName);
 	virtual void CancelRequestRemoveContentOnNextInit(TArrayView<const FName> BundleNames) = 0;
 
-	virtual void CancelUpdateContent(FName BundleName, EInstallBundleCancelFlags Flags) = 0;
+	void CancelUpdateContent(FName BundleName, EInstallBundleCancelFlags Flags);
 	virtual void CancelUpdateContent(TArrayView<const FName> BundleNames, EInstallBundleCancelFlags Flags) = 0;
 
-	virtual void PauseUpdateContent(FName BundleName) = 0;
+	void PauseUpdateContent(FName BundleName);
 	virtual void PauseUpdateContent(TArrayView<const FName> BundleNames) = 0;
 
-	virtual void ResumeUpdateContent(FName BundleName) = 0;
+	void ResumeUpdateContent(FName BundleName);
 	virtual void ResumeUpdateContent(TArrayView<const FName> BundleNames) = 0;
 
 	virtual void RequestPausedBundleCallback() = 0;
@@ -101,7 +101,7 @@ public:
 	virtual TOptional<FInstallBundleProgress> GetBundleProgress(FName BundleName) const = 0;
 
 	virtual EInstallBundleRequestFlags GetModifyableContentRequestFlags() const = 0;
-	virtual void UpdateContentRequestFlags(FName BundleName, EInstallBundleRequestFlags AddFlags, EInstallBundleRequestFlags RemoveFlags) = 0;
+	void UpdateContentRequestFlags(FName BundleName, EInstallBundleRequestFlags AddFlags, EInstallBundleRequestFlags RemoveFlags);
 	virtual void UpdateContentRequestFlags(TArrayView<const FName> BundleNames, EInstallBundleRequestFlags AddFlags, EInstallBundleRequestFlags RemoveFlags) = 0;
 
 	virtual bool IsNullInterface() const = 0;
diff --git a/Engine/Source/Runtime/NullInstallBundleManager/Source/NullInstallBundleManager.cpp b/Engine/Source/Runtime/NullInstallBundleManager/Source/NullInstallBundleManager.cpp
--- a/Engine/Source/Runtime/NullInstallBundleManager/Source/NullInstallBundleManager.cpp
+++ b/Engine/Source/Runtime/NullInstallBundleManager/Source/NullInstallBundleManager.cpp
@@ -33,24 +33,12 @@ class FNullInstallBundleManager : public IInstallBundleManager
 		return EInstallBundleManagerInitState::Succeeded;
 	}
 
-	virtual FInstallBundleRequestInfo RequestUpdateContent(FName BundleName, EInstallBundleRequestFlags Flags) override
-	{
-		FInstallBundleRequestInfo RetInfo;
-		return RetInfo;
-	}
 	virtual FInstallBundleRequestInfo RequestUpdateContent(TArrayView<const FName> BundleNames, EInstallBundleRequestFlags Flags) override
 	{
 		FInstallBundleRequestInfo RetInfo;
 		return RetInfo;
 	}
 
-	virtual void GetContentState(FName BundleName, EInstallBundleGetContentStateFlags Flags, bool bAddDependencies, FInstallBundleGetContentStateDelegate Callback, FName RequestTag) override
-	{
-		FInstallBundleCombinedContentState State;
-		State.State = EInstallBundleContentState::UpToDate;
-		Callback.ExecuteIfBound(State);
-	}
-
 	virtual void GetContentState(TArrayView<const FName> BundleNames, EInstallBundleGetContentStateFlags Flags, bool bAddDependencies, FInstallBundleGetContentStateDelegate Callback, FName RequestTag) override
 	{
 		FInstallBundleCombinedContentState State;
@@ -62,42 +50,23 @@ class FNullInstallBundleManager : public IInstallBundleManager
 	{
 	}
 
-	virtual void RequestRemoveContentOnNextInit(FName RemoveName, TArrayView<const FName> KeepNames = TArrayView<const FName>()) override
-	{
-	}
 	virtual void RequestRemoveContentOnNextInit(TArrayView<const FName> RemoveNames, TArrayView<const FName> KeepNames = TArrayView<const FName>()) override
 	{
 	}
 
-	virtual void CancelRequestRemoveContentOnNextInit(FName BundleName) override
-	{
-
-	}
 	virtual void CancelRequestRemoveContentOnNextInit(TArrayView<const FName> BundleNames) override
 	{
 	}
 
-	virtual void CancelUpdateContent(FName BundleName, EInstallBundleCancelFlags Flags) override
-	{
-
-	}
 	virtual void CancelUpdateContent(TArrayView<const FName> BundleNames, EInstallBundleCancelFlags Flags) override
 	{
 	}
 
-	virtual void PauseUpdateContent(FName BundleName) override
-	{
-
-	}
 	virtual void PauseUpdateContent(TArrayView<const FName> BundleNames) override
 	{
 
 	}
 
-	virtual void ResumeUpdateContent(FName BundleName) override
-	{
-
-	}
 	virtual void ResumeUpdateContent(TArrayView<const FName> BundleNames) override
 	{
 
@@ -116,10 +85,6 @@ class FNullInstallBundleManager : public IInstallBundleManager
 	virtual EInstallBundleRequestFlags GetModifyableContentRequestFlags() const override
 	{
 		return EInstallBundleRequestFlags::None;
-	}
-	virtual void UpdateContentRequestFlags(FName BundleName, EInstallBundleRequestFlags AddFlags, EInstallBundleRequestFlags RemoveFlags) override
-	{
-
 	}
 	virtual void UpdateContentRequestFlags(TArrayView<const FName> BundleNames, EInstallBundleRequestFlags AddFlags, EInstallBundleRequestFlags RemoveFlags) override
 	{
