Install Bundle Manger
- remove EInstallBundleContentState::InitializationError since it can never happen
- Remove FInstallBundleCombinedContentState::State - this has limited value because calling code may want to combine state in different ways
- Add utility functions to FInstallBundleCombinedContentState::GetAllBundlesHaveState and FInstallBundleCombinedContentState::GetAnyBundleHasState

diff --git a/Engine/Source/Runtime/InstallBundleManager/Private/InstallBundleTypes.cpp b/Engine/Source/Runtime/InstallBundleManager/Private/InstallBundleTypes.cpp
--- a/Engine/Source/Runtime/InstallBundleManager/Private/InstallBundleTypes.cpp
+++ b/Engine/Source/Runtime/InstallBundleManager/Private/InstallBundleTypes.cpp
@@ -43,7 +43,6 @@ const TCHAR* LexToString(EInstallBundleContentState State)
 {
 	static const TCHAR* Strings[] =
 	{
-		TEXT("InitializationError"),
 		TEXT("NotInstalled"),
 		TEXT("NeedsUpdate"),
 		TEXT("UpToDate"),
@@ -83,3 +82,32 @@ const TCHAR* LexToString(EInstallBundleStatus Status)
 	static_assert(InstallBundleUtil::CastToUnderlying(EInstallBundleStatus::Count) == UE_ARRAY_COUNT(Strings), "");
 	return Strings[InstallBundleUtil::CastToUnderlying(Status)];
 }
+
+bool FInstallBundleCombinedContentState::GetAllBundlesHaveState(EInstallBundleContentState State, TArrayView<const FName> ExcludedBundles /*= TArrayView<const FName>()*/) const
+{
+	for (const TPair<FName, FInstallBundleContentState>& Pair : IndividualBundleStates)
+	{
+		if (ExcludedBundles.Contains(Pair.Key))
+			continue;
+
+		if (Pair.Value.State != State)
+			return false;
+	}
+
+	return true;
+}
+
+bool FInstallBundleCombinedContentState::GetAnyBundleHasState(EInstallBundleContentState State, TArrayView<const FName> ExcludedBundles /*= TArrayView<const FName>()*/) const
+{
+	for (const TPair<FName, FInstallBundleContentState>& Pair : IndividualBundleStates)
+	{
+		if (ExcludedBundles.Contains(Pair.Key))
+			continue;
+
+		if (Pair.Value.State == State)
+			return true;
+	}
+
+	return false;
+}
+
diff --git a/Engine/Source/Runtime/InstallBundleManager/Public/InstallBundleTypes.h b/Engine/Source/Runtime/InstallBundleManager/Public/InstallBundleTypes.h
--- a/Engine/Source/Runtime/InstallBundleManager/Public/InstallBundleTypes.h
+++ b/Engine/Source/Runtime/InstallBundleManager/Public/InstallBundleTypes.h
@@ -41,7 +41,6 @@ INSTALLBUNDLEMANAGER_API const TCHAR* LexToString(EInstallBundleManagerInitResul
 
 enum class EInstallBundleContentState : int
 {
-	InitializationError,
 	NotInstalled,
 	NeedsUpdate,
 	UpToDate,
@@ -51,20 +50,22 @@ INSTALLBUNDLEMANAGER_API const TCHAR* LexToString(EInstallBundleContentState Sta
 
 struct INSTALLBUNDLEMANAGER_API FInstallBundleContentState
 {
-	EInstallBundleContentState State = EInstallBundleContentState::InitializationError;
+	EInstallBundleContentState State = EInstallBundleContentState::NotInstalled;
 	float Weight = 0.0f;
 	TMap<EInstallBundleSourceType, FString> Version;
 };
 
 struct INSTALLBUNDLEMANAGER_API FInstallBundleCombinedContentState
 {
-	EInstallBundleContentState State = EInstallBundleContentState::InitializationError;
 	TMap<FName, FInstallBundleContentState> IndividualBundleStates;
 	TMap<EInstallBundleSourceType, FString> CurrentVersion;
 	uint64 DownloadSize = 0;
 	uint64 InstallSize = 0;
 	uint64 InstallOverheadSize = 0;
 	uint64 FreeSpace = 0;
+
+	bool GetAllBundlesHaveState(EInstallBundleContentState State, TArrayView<const FName> ExcludedBundles = TArrayView<const FName>()) const;	
+	bool GetAnyBundleHasState(EInstallBundleContentState State, TArrayView<const FName> ExcludedBundles = TArrayView<const FName>()) const;
 };
 
 enum class EInstallBundleGetContentStateFlags : uint32
diff --git a/Engine/Source/Runtime/NullInstallBundleManager/Source/NullInstallBundleManager.cpp b/Engine/Source/Runtime/NullInstallBundleManager/Source/NullInstallBundleManager.cpp
--- a/Engine/Source/Runtime/NullInstallBundleManager/Source/NullInstallBundleManager.cpp
+++ b/Engine/Source/Runtime/NullInstallBundleManager/Source/NullInstallBundleManager.cpp
@@ -42,7 +42,6 @@ class FNullInstallBundleManager : public IInstallBundleManager
 	virtual void GetContentState(TArrayView<const FName> BundleNames, EInstallBundleGetContentStateFlags Flags, bool bAddDependencies, FInstallBundleGetContentStateDelegate Callback, FName RequestTag) override
 	{
 		FInstallBundleCombinedContentState State;
-		State.State = EInstallBundleContentState::UpToDate;
 		Callback.ExecuteIfBound(State);
 	}
 
