Adding system to track persistent install stats between launches of the app.
#tests iOS Client

diff --git a/Engine/Source/Runtime/InstallBundleManager/Private/InstallBundleUtils.cpp b/Engine/Source/Runtime/InstallBundleManager/Private/InstallBundleUtils.cpp
--- a/Engine/Source/Runtime/InstallBundleManager/Private/InstallBundleUtils.cpp
+++ b/Engine/Source/Runtime/InstallBundleManager/Private/InstallBundleUtils.cpp
@@ -6,6 +6,15 @@
 
 #include "HAL/PlatformApplicationMisc.h"
 
+#include "Containers/Ticker.h"
+
+#include "Misc/ConfigCacheIni.h"
+#include "Misc/CoreDelegates.h"
+#include "Misc/Paths.h"
+#include "Misc/FileHelper.h"
+#include "HAL/PlatformFilemanager.h"
+#include "Serialization/JsonSerializerMacros.h"
+
 namespace InstallBundleUtil
 {
 	INSTALLBUNDLEMANAGER_API FString GetAppVersion()
@@ -160,4 +169,768 @@ namespace InstallBundleUtil
 			StateStats.bOpen = false;
 		}
 	}
+
+	namespace PersistentStats
+	{
+		const FString& LexToString(ETimingStatNames InType)
+		{
+			static const FString RealTotalTime(TEXT("RealTotalTime"));
+			static const FString ActiveTotalTime(TEXT("ActiveTotalTime"));
+			static const FString EstimatedTotalBGTime(TEXT("EstimatedTotalBGTime"));
+			static const FString RealChunkDBDownloadTime(TEXT("RealChunkDBDownloadTime"));
+			static const FString ActiveChunkDBDownloadTime(TEXT("ActiveChunkDBDownloadTime"));
+			static const FString EstimatedBackgroundChunkDBDownloadTime(TEXT("EstimatedBackgroundChunkDBDownloadTime"));
+			static const FString ActiveInstallTime(TEXT("ActiveInstallTime"));
+			static const FString EstimatedBGInstallTime(TEXT("EstimatedBGInstallTime"));
+			static const FString ActivePSOTime(TEXT("ActivePSOTime"));
+			static const FString EstimatedBGPSOTime(TEXT("EstimatedBGPSOTime"));
+
+			static const FString Unknown(TEXT("<Unknown PersistentStats::ETimingStatNames Value>"));
+
+			switch (InType)
+			{
+			case ETimingStatNames::RealTotalTime:
+				return RealTotalTime;
+			case ETimingStatNames::ActiveTotalTime:
+				return ActiveTotalTime;
+			case ETimingStatNames::EstimatedTotalBGTime:
+				return EstimatedTotalBGTime;
+			case ETimingStatNames::RealChunkDBDownloadTime:
+				return RealChunkDBDownloadTime;
+			case ETimingStatNames::ActiveChunkDBDownloadTime:
+				return ActiveChunkDBDownloadTime;
+			case ETimingStatNames::EstimatedBackgroundChunkDBDownloadTime:
+				return EstimatedBackgroundChunkDBDownloadTime;
+			case ETimingStatNames::ActiveInstallTime:
+				return ActiveInstallTime;
+			case ETimingStatNames::EstimatedBGInstallTime:
+				return EstimatedBGInstallTime;
+			case ETimingStatNames::ActivePSOTime:
+				return ActivePSOTime;
+			case ETimingStatNames::EstimatedBGPSOTime:
+				return EstimatedBGPSOTime;
+			default:
+				break;
+			}
+
+			ensureAlwaysMsgf(false, TEXT("Missing PersistentStats::ETimingStatNames LexToString entry! Missing Entry as Int: %d"), (int)(InType));
+			return Unknown;
+		}
+	
+		const FString& LexToString(ECountStatNames InType)
+		{
+			static const FString NumResumedFromBackground(TEXT("NumResumedFromBackground"));
+			static const FString NumResumedFromLaunch(TEXT("NumResumedFromLaunch"));
+			static const FString NumBackgrounded(TEXT("NumBackgrounded"));
+
+			static const FString Unknown(TEXT("<Unknown PersistentStats::ETimingStatNames Value>"));
+
+			switch (InType)
+			{
+			case ECountStatNames::NumResumedFromBackground:
+				return NumResumedFromBackground;
+			case ECountStatNames::NumResumedFromLaunch:
+				return NumResumedFromLaunch;
+			case ECountStatNames::NumBackgrounded:
+				return NumBackgrounded;
+			default:
+				break;
+			}
+
+			ensureAlwaysMsgf(false, TEXT("Missing PersistentStats::ECountStatNames LexToString entry! Missing Entry as Int: %d"), (int)(InType));
+			return Unknown;
+		}
+
+		bool FPersistentStatsBase::LoadStatsFromDisk()
+		{
+			FString JSONStringOnDisk;
+			if (FPaths::FileExists(GetFullPathForStatFile()))
+			{
+				FFileHelper::LoadFileToString(JSONStringOnDisk, *GetFullPathForStatFile());
+			}
+
+			if (!JSONStringOnDisk.IsEmpty())
+			{
+				return FromJson(JSONStringOnDisk);
+			}
+
+			return false;
+		}
+
+		bool FPersistentStatsBase::SaveStatsToDisk()
+		{
+			bIsDirty = false;
+			return FFileHelper::SaveStringToFile(ToJson(), *GetFullPathForStatFile());
+		}
+
+		void FPersistentStatsBase::ResetStats(const FString& NewAnalyticsSessionID)
+		{
+			TimingStatsMap.Reset();
+			CountStatMap.Reset();
+			AnalyticsSessionID = NewAnalyticsSessionID;
+
+			bIsDirty = true;
+		}
+
+		bool FPersistentStatsBase::HasTimingStat(ETimingStatNames StatToCheck) const
+		{
+			const FPersistentTimerData* FoundStat = TimingStatsMap.Find(LexToString(StatToCheck));
+			return (nullptr != FoundStat);
+		}
+
+		bool FPersistentStatsBase::HasCountStat(ECountStatNames StatToCheck) const
+		{
+			const int* FoundStat = CountStatMap.Find(LexToString(StatToCheck));
+			return (nullptr != FoundStat);
+		}
+
+		const FPersistentTimerData* FPersistentStatsBase::GetTimingStatData(ETimingStatNames StatToGet) const
+		{ 
+			return TimingStatsMap.Find(LexToString(StatToGet));
+		}
+
+		const int* FPersistentStatsBase::GetCountStatData(ECountStatNames StatToGet) const
+		{
+			return CountStatMap.Find(LexToString(StatToGet));
+		}
+
+		void FPersistentStatsBase::IncrementCountStat(PersistentStats::ECountStatNames StatToUpdate)
+		{
+			int& StatCount = CountStatMap.FindOrAdd(LexToString(StatToUpdate));
+			++StatCount;
+
+			bIsDirty = true;
+		}
+
+		bool FPersistentStatsBase::IsTimingStatStarted(PersistentStats::ETimingStatNames StatToUpdate) const
+		{
+			bool HasStarted = false;
+
+			if (HasTimingStat(StatToUpdate))
+			{
+				const FPersistentTimerData* FoundStat = GetTimingStatData(StatToUpdate);
+				if (ensureAlwaysMsgf((nullptr != FoundStat), TEXT("Missing FInstallBundlePersistentTimingData but returned true from HasTimingStat For Stat:%s!"), *LexToString(StatToUpdate)))
+				{
+					HasStarted = (FoundStat->LastUpdateTime != 0.);
+				}
+			}
+
+			return HasStarted;
+		}
+
+		void FPersistentStatsBase::StartTimingStat(PersistentStats::ETimingStatNames StatToUpdate)
+		{
+			FPersistentTimerData& FoundStat = TimingStatsMap.FindOrAdd(LexToString(StatToUpdate));
+			FoundStat.LastUpdateTime = FPlatformTime::Seconds();
+
+			bIsDirty = true;
+		}
+
+		void FPersistentStatsBase::StopTimingStat(PersistentStats::ETimingStatNames StatToUpdate, bool UpdateTimerOnStop /* = true */)
+		{
+			//Only want to actually update the timer if we have started it (otherwise the update won't do anything and will ensure)
+			if (UpdateTimerOnStop && IsTimingStatStarted(StatToUpdate))
+			{
+				UpdateTimingStat(StatToUpdate);
+			}
+
+			FPersistentTimerData& FoundStat = TimingStatsMap.FindOrAdd(LexToString(StatToUpdate));
+			FoundStat.LastUpdateTime = 0.;
+
+			bIsDirty = true;
+		}
+
+		void FPersistentStatsBase::UpdateTimingStat(PersistentStats::ETimingStatNames StatToUpdate)
+		{
+			if (ensureAlwaysMsgf(IsTimingStatStarted(StatToUpdate), TEXT("Calling UpdateTimingStat on a stat that hasn't been started! %s"), *LexToString(StatToUpdate)))
+			{
+				FPersistentTimerData& FoundStat = TimingStatsMap.FindOrAdd(LexToString(StatToUpdate));
+
+				const double CurrentTime = FPlatformTime::Seconds();
+				const double TimeSinceUpdate = CurrentTime - FoundStat.LastUpdateTime;
+
+				if (ensureAlwaysMsgf((TimeSinceUpdate > 0.f), TEXT("Invalid saved LastUpdateTime for Stat %s! Possible Logic Error!")))
+				{
+					FoundStat.CurrentValue += TimeSinceUpdate;
+				}
+				
+				FoundStat.LastUpdateTime = CurrentTime;
+				bIsDirty = true;
+			}
+		}
+
+		void FPersistentStatsBase::UpdateAllActiveTimers()
+		{
+			for (uint8 TimingStatNameIndex = 0; TimingStatNameIndex < (uint8)PersistentStats::ETimingStatNames::NumStatNames; ++TimingStatNameIndex)
+			{
+				PersistentStats::ETimingStatNames EnumForIndex = (PersistentStats::ETimingStatNames)TimingStatNameIndex;
+				if (IsTimingStatStarted(EnumForIndex))
+				{
+					UpdateTimingStat(EnumForIndex);
+				}
+			}
+		}
+
+		void FPersistentStatsBase::StopAllActiveTimers()
+		{
+			for (uint8 TimingStatIndex = 0; TimingStatIndex < static_cast<uint8>(ETimingStatNames::NumStatNames); ++TimingStatIndex)
+			{
+				ETimingStatNames TimingStatAsEnum = static_cast<ETimingStatNames>(TimingStatIndex);
+				if (IsTimingStatStarted(TimingStatAsEnum))
+				{
+					StopTimingStat(TimingStatAsEnum);
+				}
+			}
+		}
+
+		void FPersistentStatsBase::StatsBegin(const FString& ExpectedAnalyticsID, bool bForceResetData /* = false */)
+		{
+			bIsActive = true;
+			
+			if (LoadStatsFromDisk())
+			{
+				OnLoadingDataFromDisk();
+			}
+
+			//If our Analytics ID doesn't match our expected we need to reset the data as we have started a new persistent session
+			if (bForceResetData || !AnalyticsSessionID.Equals(ExpectedAnalyticsID))
+			{
+				ResetStats(ExpectedAnalyticsID);
+			}
+			
+			//Immediately save here so we don't risk reloading the same stale data
+			//if we don't make it to an update
+			SaveStatsToDisk();
+		}
+
+		void FPersistentStatsBase::StatsEnd(bool bStopAllActiveTimers /* = true */)
+		{
+			bIsActive = false;
+
+			if (bStopAllActiveTimers)
+			{
+				StopAllActiveTimers();
+			}
+
+			//Immediately save here as we only look to update active dirty bundles, and since
+			//this likely won't be changed anymore we might as well save it out now
+			SaveStatsToDisk();
+		}
+
+		void FPersistentStatsBase::OnLoadingDataFromDisk()
+		{
+			HandleTimerStatsAfterDataLoad();
+		}
+
+		void FPersistentStatsBase::HandleTimerStatsAfterDataLoad()
+		{
+			//Go through all timing stats and handle each one accordingly
+			//All Active Timers we should stop these timers without updating them so that they don't accrue time from backgrounding
+			//All BG timers Should be stopped, but update their timers on stopping so they do accrue time from backgrounding
+
+			/*
+					Handle Active (Foreground) Timers
+			*/
+
+			if (IsTimingStatStarted(ETimingStatNames::ActiveTotalTime))
+			{
+				StopTimingStat(ETimingStatNames::ActiveTotalTime, false);
+			}
+
+			if (IsTimingStatStarted(ETimingStatNames::ActiveChunkDBDownloadTime))
+			{
+				StopTimingStat(ETimingStatNames::ActiveChunkDBDownloadTime, false);
+			}
+
+			if (IsTimingStatStarted(ETimingStatNames::ActiveInstallTime))
+			{
+				StopTimingStat(ETimingStatNames::ActiveInstallTime, false);
+			}
+
+			if (IsTimingStatStarted(ETimingStatNames::ActivePSOTime))
+			{
+				StopTimingStat(ETimingStatNames::ActivePSOTime, false);
+			}
+
+			/*
+					Handle Real Total Timers
+			*/
+			
+			// RealTotalTime should always be updated, it is an attempted tracking of all time (BG and FG) so treat it like a BG timer in this case to accrue background time
+			// We stop it here instead of just updating it as we don't want to accrue time until it is started again (just account for the BG time)
+			if (IsTimingStatStarted(ETimingStatNames::RealTotalTime))
+			{
+				StopTimingStat(ETimingStatNames::RealTotalTime, true);
+			}
+			
+			if (IsTimingStatStarted(ETimingStatNames::RealChunkDBDownloadTime))
+			{
+				StopTimingStat(ETimingStatNames::RealChunkDBDownloadTime, true);
+			}
+
+			/*
+					Handle Background Timers
+			*/
+			
+			if (IsTimingStatStarted(ETimingStatNames::EstimatedTotalBGTime))
+			{
+				StopTimingStat(ETimingStatNames::EstimatedTotalBGTime, true);
+			}
+
+			if (IsTimingStatStarted(ETimingStatNames::EstimatedBackgroundChunkDBDownloadTime))
+			{
+				StopTimingStat(ETimingStatNames::EstimatedBackgroundChunkDBDownloadTime, true);
+			}
+
+			if (IsTimingStatStarted(ETimingStatNames::EstimatedBGInstallTime))
+			{
+				StopTimingStat(ETimingStatNames::EstimatedBGInstallTime, true);
+			}
+
+			if (IsTimingStatStarted(ETimingStatNames::EstimatedBGPSOTime))
+			{
+				StopTimingStat(ETimingStatNames::EstimatedBGPSOTime, true);
+			}
+		}
+
+		void FSessionPersistentStats::AddRequiredBundles(const TArray<FString>& RequiredBundlesToAdd)
+		{
+			for (const FString& BundleName : RequiredBundlesToAdd)
+			{
+				RequiredBundles.AddUnique(BundleName);
+			}
+			bIsDirty = true;
+		}
+
+		void FSessionPersistentStats::AddRequiredBundles(const TArray<FName>& RequiredBundlesToAdd)
+		{
+			for (FName BundleName : RequiredBundlesToAdd)
+			{
+				RequiredBundles.AddUnique(BundleName.ToString());
+			}
+			bIsDirty = true;
+		}
+
+		void FSessionPersistentStats::GetRequiredBundles(TArray<FString>& OutRequiredBundles)
+		{
+			OutRequiredBundles.Empty();
+			for (const FString& BundleName : RequiredBundles)
+			{
+				OutRequiredBundles.Add(BundleName);
+			}
+		}
+	
+		void FSessionPersistentStats::ResetRequiredBundles(const TArray<FString>& NewRequiredBundles /* = TArray<FString>() */)
+		{
+			RequiredBundles.Empty();
+			AddRequiredBundles(NewRequiredBundles);
+			bIsDirty = true;
+		}
+
+		const FString FBundlePersistentStats::GetFullPathForStatFile() const
+		{
+			return FPaths::Combine(FPlatformMisc::GamePersistentDownloadDir(), TEXT("PersistentStats"), TEXT("BundleStats"), (BundleName + TEXT(".json")));
+		}
+
+		const FString FSessionPersistentStats::GetFullPathForStatFile() const
+		{
+			return FPaths::Combine(FPlatformMisc::GamePersistentDownloadDir(), TEXT("PersistentStats"), TEXT("ContentRequestStats"), (SessionName + TEXT(".json")));
+		}
+
+		void FPersistentStatContainerBase::InitializeBase()
+		{
+			//Load Settings from Config
+			{
+				GConfig->GetBool(TEXT("InstallBundleManager.PersistentStatSettings"), TEXT("ShouldAutoUpdateInTick"), bShouldAutoUpdateInTick, GEngineIni);
+				GConfig->GetBool(TEXT("InstallBundleManager.PersistentStatSettings"), TEXT("ShouldAutoUpdateBackgroundStats"), bShouldAutoUpdateBackgroundStats, GEngineIni);
+				GConfig->GetBool(TEXT("InstallBundleManager.PersistentStatSettings"), TEXT("bShouldSaveDirtyStatsOnTick"), bShouldSaveDirtyStatsOnTick, GEngineIni);
+				GConfig->GetBool(TEXT("InstallBundleManager.PersistentStatSettings"), TEXT("bShouldSaveStatsEveryUpdate"), bShouldSaveStatsEveryUpdate, GEngineIni);
+
+				GConfig->GetFloat(TEXT("InstallBundleManager.PersistentStatSettings"), TEXT("AutoUpdateRate"), TimerStat_ResetTimerValue, GEngineIni);
+				ResetTimerUpdate();
+			}
+
+			//Setup Delegates (Needs to happen after config to have AutoUpdate settings loaded)
+			{
+				if ((bShouldAutoUpdateInTick || bShouldSaveDirtyStatsOnTick) && !TickHandle.IsValid())
+				{
+					TickHandle = FTicker::GetCoreTicker().AddTicker(FTickerDelegate::CreateRaw(this, &FPersistentStatContainerBase::Tick));
+				}
+
+				if (bShouldAutoUpdateBackgroundStats)
+				{
+					if (!OnApp_EnteringForegroundHandle.IsValid())
+					{
+						OnApp_EnteringForegroundHandle = FCoreDelegates::ApplicationHasEnteredForegroundDelegate.AddRaw(this, &FPersistentStatContainerBase::OnApp_EnteringForeground);
+					}
+					if (!OnApp_EnteringBackgroundHandle.IsValid())
+					{
+						OnApp_EnteringBackgroundHandle = FCoreDelegates::ApplicationWillEnterBackgroundDelegate.AddRaw(this, &FPersistentStatContainerBase::OnApp_EnteringBackground);
+					}
+				}
+			}
+		}
+
+		void FPersistentStatContainerBase::ShutdownBase()
+		{
+			if (TickHandle.IsValid())
+			{
+				FTicker::GetCoreTicker().RemoveTicker(TickHandle);
+				TickHandle.Reset();
+			}
+
+			if (OnApp_EnteringForegroundHandle.IsValid())
+			{
+				FCoreDelegates::ApplicationHasEnteredForegroundDelegate.Remove(OnApp_EnteringForegroundHandle);
+				OnApp_EnteringForegroundHandle.Reset();
+			}
+			if (OnApp_EnteringBackgroundHandle.IsValid())
+			{
+				FCoreDelegates::ApplicationWillEnterBackgroundDelegate.Remove(OnApp_EnteringBackgroundHandle);
+				OnApp_EnteringBackgroundHandle.Reset();
+			}
+		}
+
+		bool FPersistentStatContainerBase::Tick(float dt)
+		{
+			if (bShouldAutoUpdateInTick)
+			{
+				//Only update all active timers every TimerStat_ResetTimerValue seconds
+				TimerStat_UpdateTimer -= dt;
+				if (TimerStat_UpdateTimer < 0.0f)
+				{
+					ResetTimerUpdate();
+					UpdateAllBundlesActiveTimers();
+					UpdateAllSessionActiveTimers();
+				}
+			}
+
+			if (bShouldSaveDirtyStatsOnTick)
+			{
+				//Go ahead and update all dirty stats every tick
+				//This way we don't save them multiple times in a single tick
+				//for updating multiple stats
+				SaveAllDirtyStatsToDisk();
+			}
+			//go ahead and always Tick once we start
+			return true;
+		}
+
+		void FPersistentStatContainerBase::OnDataUpdatedForStat(FPersistentStatsBase& UpdatedBundleStat)
+		{
+			if (bShouldSaveStatsEveryUpdate)
+			{
+				UpdatedBundleStat.SaveStatsToDisk();
+			}
+		}
+
+		void FPersistentStatContainerBase::ResetTimerUpdate()
+		{
+			TimerStat_UpdateTimer = TimerStat_ResetTimerValue;
+		}
+
+		void FPersistentStatContainerBase::SaveAllDirtyStatsToDisk()
+		{
+			TArray<FName> AllBundleStatNames;
+			PerBundlePersistentStatMap.GetKeys(AllBundleStatNames);
+			for (FName& BundleName : AllBundleStatNames)
+			{
+				FBundlePersistentStats* BundleStats = PerBundlePersistentStatMap.Find(BundleName);
+				check(BundleStats);
+				if (BundleStats->IsDirty())
+				{
+					BundleStats->SaveStatsToDisk();
+				}
+			}
+
+			TArray<FString> AllSessionStatNames;
+			SessionPersistentStatMap.GetKeys(AllSessionStatNames);
+			for (const FString& SessionName : AllSessionStatNames)
+			{
+				FSessionPersistentStats* SessionStats = SessionPersistentStatMap.Find(SessionName);
+				check(SessionStats);
+				if (SessionStats->IsDirty())
+				{
+					SessionStats->SaveStatsToDisk();
+				}
+			}			
+		}
+
+
+
+		void FPersistentStatContainerBase::StartBundlePersistentStatTracking(FName BundleName, const FString& ExpectedAnalyticsID /* = FString() */, bool bForceResetStatData /* = false */)
+		{
+			//Use the base expected analytics ID if one was not passed in
+			const FString ExpectedAnalyticsToUse = ExpectedAnalyticsID.IsEmpty() ? FPersistentStatsBase::GetBaseExpectedAnalyticsID() : ExpectedAnalyticsID;
+
+			FBundlePersistentStats& FoundBundleStats = PerBundlePersistentStatMap.FindOrAdd(BundleName, FBundlePersistentStats(BundleName));
+			FoundBundleStats.StatsBegin(ExpectedAnalyticsToUse, bForceResetStatData);
+
+			OnDataUpdatedForStat(FoundBundleStats);
+		}
+
+		void FPersistentStatContainerBase::StartSessionPersistentStatTracking(const FString& SessionName, const TArray<FName>& RequiredBundles /* = TArray<FName>() */, const FString& ExpectedAnalyticsID /* = FString() */, bool bForceResetStatData /* = false */)
+		{
+			//Use the base expected analytics ID if one was not passed in
+			const FString ExpectedAnalyticsToUse = ExpectedAnalyticsID.IsEmpty() ? FPersistentStatsBase::GetBaseExpectedAnalyticsID() : ExpectedAnalyticsID;
+
+			FSessionPersistentStats& FoundSessionStats = SessionPersistentStatMap.FindOrAdd(SessionName, FSessionPersistentStats(SessionName));
+			FoundSessionStats.StatsBegin(ExpectedAnalyticsToUse, bForceResetStatData);
+
+			//Also append starting required bundles as we may have new ones from the ones already in data
+			FoundSessionStats.AddRequiredBundles(RequiredBundles);
+
+			OnDataUpdatedForStat(FoundSessionStats);
+		}
+		
+		void FPersistentStatContainerBase::StopBundlePersistentStatTracking(FName BundleName, bool bStopAllActiveTimers /* = true */)
+		{
+			FBundlePersistentStats* FoundBundleStats = PerBundlePersistentStatMap.Find(BundleName);
+			if (nullptr != FoundBundleStats)
+			{
+				FoundBundleStats->StatsEnd(bStopAllActiveTimers);
+				OnDataUpdatedForStat(*FoundBundleStats);
+			}
+		}
+
+		void FPersistentStatContainerBase::StopSessionPersistentStatTracking(FString SessionName, bool bStopAllActiveTimers /* = true */)
+		{
+			FSessionPersistentStats* FoundSessionStats = SessionPersistentStatMap.Find(SessionName);
+			if (nullptr != FoundSessionStats)
+			{
+				FoundSessionStats->StatsEnd(bStopAllActiveTimers);
+				OnDataUpdatedForStat(*FoundSessionStats);
+			}
+		}
+
+		void FPersistentStatContainerBase::StartBundlePersistentStatTimer(FName BundleName, ETimingStatNames TimerToStart)
+		{
+			FBundlePersistentStats& FoundBundleStats = PerBundlePersistentStatMap.FindOrAdd(BundleName, FBundlePersistentStats(BundleName));
+			FoundBundleStats.StartTimingStat(TimerToStart);
+
+			OnDataUpdatedForStat(FoundBundleStats);
+		}
+
+		void FPersistentStatContainerBase::StartSessionPersistentStatTimer(FString SessionName, ETimingStatNames TimerToStart)
+		{
+			FSessionPersistentStats& FoundSessionStats = SessionPersistentStatMap.FindOrAdd(SessionName, FSessionPersistentStats(SessionName));
+			FoundSessionStats.StartTimingStat(TimerToStart);
+
+			OnDataUpdatedForStat(FoundSessionStats);
+		}
+
+		void FPersistentStatContainerBase::StopBundlePersistentStatTimer(FName BundleName, ETimingStatNames TimerToStop)
+		{
+			FBundlePersistentStats& FoundBundleStats = PerBundlePersistentStatMap.FindOrAdd(BundleName, FBundlePersistentStats(BundleName));
+			FoundBundleStats.StopTimingStat(TimerToStop);
+
+			OnDataUpdatedForStat(FoundBundleStats);
+		}
+
+		void FPersistentStatContainerBase::StopSessionPersistentStatTimer(FString SessionName, ETimingStatNames TimerToStop)
+		{
+			FSessionPersistentStats& FoundSessionStats = SessionPersistentStatMap.FindOrAdd(SessionName, FSessionPersistentStats(SessionName));
+			FoundSessionStats.StopTimingStat(TimerToStop);
+
+			OnDataUpdatedForStat(FoundSessionStats);
+		}
+
+		void FPersistentStatContainerBase::UpdateBundlePersistentStatTimer(FName BundleName, ETimingStatNames TimerToUpdate)
+		{
+			FBundlePersistentStats& FoundBundleStats = PerBundlePersistentStatMap.FindOrAdd(BundleName, FBundlePersistentStats(BundleName));
+			FoundBundleStats.UpdateTimingStat(TimerToUpdate);
+
+			OnDataUpdatedForStat(FoundBundleStats);
+		}
+
+		void FPersistentStatContainerBase::UpdateSessionPersistentStatTimer(FString SessionName, ETimingStatNames TimerToUpdate)
+		{
+			FSessionPersistentStats& FoundSessionStats = SessionPersistentStatMap.FindOrAdd(SessionName, FSessionPersistentStats(SessionName));
+			FoundSessionStats.UpdateTimingStat(TimerToUpdate);
+
+			OnDataUpdatedForStat(FoundSessionStats);
+		}
+
+		void FPersistentStatContainerBase::IncrementBundlePersistentCounter(FName BundleName, ECountStatNames CounterToUpdate)
+		{
+			FBundlePersistentStats& FoundBundleStats = PerBundlePersistentStatMap.FindOrAdd(BundleName, FBundlePersistentStats(BundleName));
+			FoundBundleStats.IncrementCountStat(CounterToUpdate);
+
+			OnDataUpdatedForStat(FoundBundleStats);
+		}
+
+		void FPersistentStatContainerBase::IncrementSessionPersistentCounter(FString SessionName, ECountStatNames CounterToUpdate)
+		{
+			FSessionPersistentStats& FoundSessionStats = SessionPersistentStatMap.FindOrAdd(SessionName, FSessionPersistentStats(SessionName));
+			FoundSessionStats.IncrementCountStat(CounterToUpdate);
+
+			OnDataUpdatedForStat(FoundSessionStats);
+		}
+
+		void FPersistentStatContainerBase::OnApp_EnteringBackground()
+		{
+			OnBackground_HandleBundleStats();
+			OnBackground_HandleSessionStats();
+		}
+
+		void FPersistentStatContainerBase::OnApp_EnteringForeground()
+		{
+			OnForeground_HandleBundleStats();
+			OnForeground_HandleSessionStats();
+		}
+
+		void FPersistentStatContainerBase::OnBackground_HandleBundleStats()
+		{
+			for (TPair<FName, FBundlePersistentStats>& BundlePair : PerBundlePersistentStatMap)
+			{
+				//Only bother updating bundles listed as active
+				if (BundlePair.Value.IsActive())
+				{
+					UpdateStatsForBackground(BundlePair.Value);
+				}
+			}
+		}
+
+		void FPersistentStatContainerBase::OnForeground_HandleBundleStats()
+		{
+			for (TPair<FName, FBundlePersistentStats>& BundlePair : PerBundlePersistentStatMap)
+			{
+				//Only bother updating bundles listed as active
+				if (BundlePair.Value.IsActive())
+				{
+					UpdateStatsForForeground(BundlePair.Value);
+				}
+			}
+		}
+
+		void FPersistentStatContainerBase::OnBackground_HandleSessionStats()
+		{
+			for (TPair<FString, FSessionPersistentStats>& SessionPair : SessionPersistentStatMap)
+			{
+				//Only bother updating sessions listed as active
+				if (SessionPair.Value.IsActive())
+				{
+					UpdateStatsForBackground(SessionPair.Value);
+				}
+			}
+		}
+
+		void FPersistentStatContainerBase::OnForeground_HandleSessionStats()
+		{
+			for (TPair<FString, FSessionPersistentStats>& SessionPair : SessionPersistentStatMap)
+			{
+				//Only bother updating sessions listed as active
+				if (SessionPair.Value.IsActive())
+				{
+					UpdateStatsForForeground(SessionPair.Value);
+				}
+			}
+		}
+
+		void FPersistentStatContainerBase::UpdateStatsForBackground(FPersistentStatsBase& StatToUpdate)
+		{
+			StatToUpdate.IncrementCountStat(ECountStatNames::NumBackgrounded);
+
+			//Always handle ActiveTotalTime as this isn't dependent on what stage of the process we are in
+			if (StatToUpdate.IsTimingStatStarted(ETimingStatNames::ActiveTotalTime))
+			{
+				StatToUpdate.StartTimingStat(ETimingStatNames::EstimatedTotalBGTime);
+				StatToUpdate.StopTimingStat(ETimingStatNames::ActiveTotalTime);
+			}
+
+			//Besides the ActiveTotalTime above, we should only be in 1 of the following states at a time, so only handle the appropriate swap
+			if (StatToUpdate.IsTimingStatStarted(ETimingStatNames::ActiveChunkDBDownloadTime))
+			{
+				StatToUpdate.StartTimingStat(ETimingStatNames::EstimatedBackgroundChunkDBDownloadTime);
+				StatToUpdate.StopTimingStat(ETimingStatNames::ActiveChunkDBDownloadTime);
+			}
+			else if (StatToUpdate.IsTimingStatStarted(ETimingStatNames::ActiveInstallTime))
+			{
+				StatToUpdate.StartTimingStat(ETimingStatNames::EstimatedBGInstallTime);
+				StatToUpdate.StopTimingStat(ETimingStatNames::ActiveInstallTime);
+			}
+			else if (StatToUpdate.IsTimingStatStarted(ETimingStatNames::ActivePSOTime))
+			{
+				StatToUpdate.StartTimingStat(ETimingStatNames::EstimatedBGPSOTime);
+				StatToUpdate.StopTimingStat(ETimingStatNames::ActivePSOTime);
+			}
+
+			OnDataUpdatedForStat(StatToUpdate);
+		}
+
+		void FPersistentStatContainerBase::UpdateStatsForForeground(FPersistentStatsBase& StatToUpdate)
+		{
+			StatToUpdate.IncrementCountStat(ECountStatNames::NumResumedFromBackground);
+
+			//Always handle ActiveTotalTime as this isn't dependent on what stage of the process we are in
+			if (StatToUpdate.IsTimingStatStarted(ETimingStatNames::EstimatedTotalBGTime))
+			{
+				StatToUpdate.StopTimingStat(ETimingStatNames::EstimatedTotalBGTime);
+				StatToUpdate.StartTimingStat(ETimingStatNames::ActiveTotalTime);
+			}
+
+			//Besides the ActiveTotalTime above, we should only be in 1 of the following states at a time, so only handle the appropriate swap
+			if (StatToUpdate.IsTimingStatStarted(ETimingStatNames::EstimatedBackgroundChunkDBDownloadTime))
+			{
+				StatToUpdate.StopTimingStat(ETimingStatNames::EstimatedBackgroundChunkDBDownloadTime);
+				StatToUpdate.StartTimingStat(ETimingStatNames::ActiveChunkDBDownloadTime);
+			}
+			else if (StatToUpdate.IsTimingStatStarted(ETimingStatNames::EstimatedBGInstallTime))
+			{
+				StatToUpdate.StopTimingStat(ETimingStatNames::EstimatedBGInstallTime);
+				StatToUpdate.StartTimingStat(ETimingStatNames::ActiveInstallTime);
+			}
+			else if (StatToUpdate.IsTimingStatStarted(ETimingStatNames::EstimatedBGPSOTime))
+			{
+				StatToUpdate.StopTimingStat(ETimingStatNames::EstimatedBGPSOTime);
+				StatToUpdate.StartTimingStat(ETimingStatNames::ActivePSOTime);
+			}
+
+			OnDataUpdatedForStat(StatToUpdate);
+		}
+
+		void FPersistentStatContainerBase::UpdateAllBundlesActiveTimers()
+		{
+			TArray<FName> BundleNames;
+			PerBundlePersistentStatMap.GetKeys(BundleNames);
+			
+			for (FName BundleName : BundleNames)
+			{
+				InstallBundleUtil::PersistentStats::FBundlePersistentStats& BundleStats = PerBundlePersistentStatMap.FindOrAdd(BundleName, FBundlePersistentStats(BundleName));
+				BundleStats.UpdateAllActiveTimers();
+
+				OnDataUpdatedForStat(BundleStats);
+			}
+		}
+
+		void FPersistentStatContainerBase::UpdateAllSessionActiveTimers()
+		{
+			TArray<FString> SessionNames;
+			SessionPersistentStatMap.GetKeys(SessionNames);
+
+			for (const FString& SessionName : SessionNames)
+			{
+				InstallBundleUtil::PersistentStats::FSessionPersistentStats& SessionStats = SessionPersistentStatMap.FindOrAdd(SessionName, FSessionPersistentStats(SessionName));
+				SessionStats.UpdateAllActiveTimers();
+
+				OnDataUpdatedForStat(SessionStats);
+			}
+		}
+
+		const FString FPersistentStatsBase::GetBaseExpectedAnalyticsID()
+		{
+			const FString BaseExpectedAnalyticsID = FPlatformMisc::GetDeviceId() + TEXT("_") + FApp::GetBuildVersion();
+			return BaseExpectedAnalyticsID;
+		}
+
+		const FBundlePersistentStats* FPersistentStatContainerBase::GetBundleStat(FName BundleName) const
+		{
+			return PerBundlePersistentStatMap.Find(BundleName);
+		}
+
+		const FSessionPersistentStats* FPersistentStatContainerBase::GetSessionStat(const FString& SessionName) const
+		{
+			return SessionPersistentStatMap.Find(SessionName);
+		}
+	}
 }
diff --git a/Engine/Source/Runtime/InstallBundleManager/Public/InstallBundleManagerInterface.h b/Engine/Source/Runtime/InstallBundleManager/Public/InstallBundleManagerInterface.h
--- a/Engine/Source/Runtime/InstallBundleManager/Public/InstallBundleManagerInterface.h
+++ b/Engine/Source/Runtime/InstallBundleManager/Public/InstallBundleManagerInterface.h
@@ -109,5 +109,8 @@ public:
 	virtual void SetErrorSimulationCommands(const FString& CommandLine) {}
 
 	virtual TSharedPtr<IAnalyticsProviderET> GetAnalyticsProvider() const { return TSharedPtr<IAnalyticsProviderET>(); }
+
+	virtual void StartPersistentStatTrackingSession(const FString& SessionName, const TArray<FName>& RequiredBundles = TArray<FName>(), const FString& ExpectedAnalyticsID = FString(), bool bForceResetStatData = false) {}
+	virtual void StopPersistentStatTrackingSession(const FString& SessionName) {}
 };
 
diff --git a/Engine/Source/Runtime/InstallBundleManager/Public/InstallBundleSourceInterface.h b/Engine/Source/Runtime/InstallBundleManager/Public/InstallBundleSourceInterface.h
--- a/Engine/Source/Runtime/InstallBundleManager/Public/InstallBundleSourceInterface.h
+++ b/Engine/Source/Runtime/InstallBundleManager/Public/InstallBundleSourceInterface.h
@@ -50,6 +50,7 @@ public:
 		FInstallBundlePausedDelegate PausedCallback;
 		FInstallBundleCompleteDelegate CompleteCallback;
 		TSharedPtr<InstallBundleUtil::FContentRequestSharedContext> RequestSharedContext;
+		TSharedPtr<InstallBundleUtil::PersistentStats::FPersistentStatContainerBase> PersistentStats;
 	};
 
 	// Updates content on disk if necessary
diff --git a/Engine/Source/Runtime/InstallBundleManager/Public/InstallBundleTypes.h b/Engine/Source/Runtime/InstallBundleManager/Public/InstallBundleTypes.h
--- a/Engine/Source/Runtime/InstallBundleManager/Public/InstallBundleTypes.h
+++ b/Engine/Source/Runtime/InstallBundleManager/Public/InstallBundleTypes.h
@@ -107,6 +107,7 @@ enum class EInstallBundleRequestFlags : uint32
 	UseBackgroundDownloads = (1 << 1),
 	SendNotificationIfDownloadCompletesInBackground = (1 << 2),
 	ForceNoPatching = (1 << 3),
+	TrackPersistentBundleStats = (1 << 4),
 	Defaults = UseBackgroundDownloads,
 };
 ENUM_CLASS_FLAGS(EInstallBundleRequestFlags)
diff --git a/Engine/Source/Runtime/InstallBundleManager/Public/InstallBundleUtils.h b/Engine/Source/Runtime/InstallBundleManager/Public/InstallBundleUtils.h
--- a/Engine/Source/Runtime/InstallBundleManager/Public/InstallBundleUtils.h
+++ b/Engine/Source/Runtime/InstallBundleManager/Public/InstallBundleUtils.h
@@ -6,6 +6,7 @@
 #include "InstallBundleTypes.h"
 #include "Async/AsyncWork.h"
 #include "Misc/EmbeddedCommunication.h"
+#include "Serialization/JsonSerializerMacros.h"
 
 namespace InstallBundleUtil
 {
@@ -168,6 +169,320 @@ namespace InstallBundleUtil
 		const TMap<FName, InstallBundleUtil::FContentRequestStats>& GetMap() { return StatsMap; }
 	};
 
+	namespace PersistentStats
+	{
+		//Stats used to store timing information in FInstallBundleIndividualPersistentTimingData format
+		enum class ETimingStatNames : uint8
+		{
+			  RealTotalTime
+			, ActiveTotalTime
+			, EstimatedTotalBGTime
+			, RealChunkDBDownloadTime
+			, ActiveChunkDBDownloadTime
+			, EstimatedBackgroundChunkDBDownloadTime
+			, ActiveInstallTime
+			, EstimatedBGInstallTime
+			, ActivePSOTime
+			, EstimatedBGPSOTime 
+			, NumStatNames //should always be final entry
+		};
+
+		//Stats used to store count information in an int format
+		enum class ECountStatNames : uint8
+		{
+			  NumResumedFromBackground
+			, NumResumedFromLaunch
+			, NumBackgrounded
+			, NumStatNames //should always be final entry
+		};
+
+		const FString& LexToString(ETimingStatNames InType);
+		const FString& LexToString(ECountStatNames InType);
+
+		class FPersistentTimerData : public FJsonSerializable
+		{
+		public:
+			BEGIN_JSON_SERIALIZER
+				JSON_SERIALIZE("LastUpdateTime", LastUpdateTime);
+				JSON_SERIALIZE("TotalTime", CurrentValue);
+			END_JSON_SERIALIZER
+
+		public:
+			FPersistentTimerData()
+				: LastUpdateTime(0.)
+				, CurrentValue(0.)
+			{}
+
+			double LastUpdateTime;
+			double CurrentValue;
+		};
+
+		//Both ContentRequests and Individual Bundles use this to track data respective to themselves.
+		//Keeps us from having to write serialization / loading data for both cases that would be the same
+		class INSTALLBUNDLEMANAGER_API FPersistentStatsBase : public FJsonSerializable
+		{
+		public:
+			//We should not ever really be serializing this base class. Kept it here for reference.
+			BEGIN_JSON_SERIALIZER
+				JSON_SERIALIZE("AnalyticsSessionID", AnalyticsSessionID);
+				JSON_SERIALIZE_MAP("CountStats", CountStatMap);
+				JSON_SERIALIZE_MAP_SERIALIZABLE("TimingStats", TimingStatsMap, FPersistentTimerData);
+			END_JSON_SERIALIZER
+		
+		public:
+			virtual ~FPersistentStatsBase() {}
+
+			//Tries to load persistent stats from disk and determine if we need to resume our previous persistent session or start a new one
+			void StatsBegin(const FString& ExpectedAnalyticsID, bool bForceResetData = false);
+
+			//Tries to set this stat pool as "inactive". By default also stop all active timers. 
+			// NOTE: If you don't stop all active timers future calls to UpdateTimingStat will still increment this value and save it off!
+			void StatsEnd(bool bStopAllActiveTimers = true);
+
+			bool IsTimingStatStarted(ETimingStatNames StatToUpdate) const;
+			void StartTimingStat(ETimingStatNames StatToUpdate);
+			void StopTimingStat(ETimingStatNames StatToUpdate, bool UpdateTimerOnStop = true);
+			void UpdateTimingStat(ETimingStatNames StatToUpdate);
+			
+			void UpdateAllActiveTimers();
+			void StopAllActiveTimers();
+			
+			bool IsActive() const { return bIsActive; }
+			bool IsDirty() const { return bIsDirty; }
+
+			void IncrementCountStat(PersistentStats::ECountStatNames StatToUpdate);
+
+			bool HasTimingStat(ETimingStatNames StatToCheck) const;
+			bool HasCountStat(ECountStatNames StatToCheck) const;
+
+			const FPersistentTimerData* GetTimingStatData(ETimingStatNames StatToGet) const;
+			const int* GetCountStatData(ECountStatNames StatToGet) const;
+
+			//Saves the persistent stat data to disk in the location returned by GetFullPathForStatFile()
+			//Returns True if that save succeeds and false otherwise.
+			bool SaveStatsToDisk();
+
+		protected:
+			bool LoadStatsFromDisk();
+			void ResetStats(const FString& NewAnalyticsSessionID);
+
+			//Called after we load data but before we process any of it.
+			//Useful for fixing up the data loaded from disk before it changes or is acted upon by anything else
+			void OnLoadingDataFromDisk();
+
+			//Helper to try and reconcile offline and active timers after we load data from disk
+			void HandleTimerStatsAfterDataLoad();
+
+			//Protected so we can only construct the derived versions of this Base class
+			FPersistentStatsBase()
+				: TimingStatsMap()
+				, CountStatMap()
+				, AnalyticsSessionID()
+				, bIsActive(false)
+			{
+			}
+			
+		protected:
+			//Serialized Stat Data
+			TMap<FString, FPersistentTimerData> TimingStatsMap;
+			TMap<FString, int> CountStatMap;
+			FString AnalyticsSessionID;
+
+			//We don't serialize these as they are tracking behavior and not stats
+			bool bIsActive;
+			bool bIsDirty;
+			
+		//Static Methods
+		public:
+			//Helper function that generates a basic ExpectedAnalyticsID from device and UE4 information
+			static const FString GetBaseExpectedAnalyticsID();
+		
+		//Pure Virtual Methods
+		public:
+			//Returns Full Path to look for our persistent stat cache
+			virtual const FString GetFullPathForStatFile() const = 0;
+		};
+
+		//Each bundles persistent stats
+		class INSTALLBUNDLEMANAGER_API FBundlePersistentStats : public FPersistentStatsBase
+		{
+		public:
+			BEGIN_JSON_SERIALIZER
+				JSON_SERIALIZE("BundleName", BundleName);
+				JSON_SERIALIZE("AnalyticsSessionID", AnalyticsSessionID);
+				JSON_SERIALIZE_MAP("CountStats", CountStatMap);
+				JSON_SERIALIZE_MAP_SERIALIZABLE("TimingStats", TimingStatsMap, FPersistentTimerData);
+			END_JSON_SERIALIZER
+
+
+		private:
+			//Remove default constructor as we require a BundleName in the constructor
+			FBundlePersistentStats() = delete;
+
+			//Used in our full path and should signify the FString version of the FName of this bundle we are tracking
+			FString BundleName;
+
+		public:
+			FBundlePersistentStats(FName BundleNameIn)
+				: BundleName(BundleNameIn.ToString())
+			{
+			}
+
+			virtual ~FBundlePersistentStats() {}
+
+			//Implementation for FInstallBundleStatsContainerBase
+		public:
+			virtual const FString GetFullPathForStatFile() const override;
+		};
+
+		//Tracks a set of bundles and stats for a named session. This way you can have stats across multiple bundle installs grouped by some naming convention
+		class INSTALLBUNDLEMANAGER_API FSessionPersistentStats : public FPersistentStatsBase
+		{
+		public:
+			BEGIN_JSON_SERIALIZER
+				JSON_SERIALIZE("SessionName", SessionName);
+				JSON_SERIALIZE_ARRAY("RequiredBundles", RequiredBundles);
+				JSON_SERIALIZE("AnalyticsSessionID", AnalyticsSessionID);
+				JSON_SERIALIZE_MAP("CountStats", CountStatMap);
+				JSON_SERIALIZE_MAP_SERIALIZABLE("TimingStats", TimingStatsMap, FPersistentTimerData);
+			END_JSON_SERIALIZER
+
+		private:
+			//Remove default constructor as we require a ContentRequestName in the constructor
+			FSessionPersistentStats() = delete;
+
+			//Used in our full path and should signify some identifier for a given ContentRequest
+			FString SessionName;
+
+			//Used to keep track of all bundles that are part of this content request
+			TArray<FString> RequiredBundles;
+
+		public:
+			FSessionPersistentStats(const FString& ContentRequestNameIn)
+				: SessionName(ContentRequestNameIn)
+				, RequiredBundles()
+			{
+			}
+
+			virtual ~FSessionPersistentStats() {}
+
+			void AddRequiredBundles(const TArray<FString>& RequiredBundlesToAdd);
+			void AddRequiredBundles(const TArray<FName>& RequiredBundlesToAdd);
+
+			//Removes all Required Bundle Names except the ones in the NewRequiredBundlesList. If no NewRequiredBundlesList is supplied, just erases all RequiredBundles.
+			//This includes ones loaded from the persistent cache from previous runs. This re-saves the persistent cache file so they will be removed completely.
+			void ResetRequiredBundles(const TArray<FString>& NewRequiredBundles = TArray<FString>());
+
+			//Populates the passed in TArray with all the required bundles' names
+			void GetRequiredBundles(TArray<FString>& OutRequiredBundles);
+			
+			//Implementation for FInstallBundleStatsContainerBase
+		public:
+			virtual const FString GetFullPathForStatFile() const override;
+		};
+
+		//Helper class to handle holding a collection of Bundle and Session stats.
+		class INSTALLBUNDLEMANAGER_API FPersistentStatContainerBase
+		{
+		public:
+			FPersistentStatContainerBase()
+				: PerBundlePersistentStatMap()
+				, SessionPersistentStatMap()
+				, TickHandle()
+				, OnApp_EnteringForegroundHandle()
+				, OnApp_EnteringBackgroundHandle()
+				, TimerStat_ResetTimerValue(5.0f)
+				, TimerStat_UpdateTimer(5.0f)
+				, bShouldAutoUpdateInTick(true)
+				, bShouldAutoUpdateBackgroundStats(true)
+				, bShouldSaveDirtyStatsOnTick(true)
+				, bShouldSaveStatsEveryUpdate(false)
+			{
+				InitializeBase();
+			}
+
+			virtual ~FPersistentStatContainerBase() 
+			{ 
+				ShutdownBase();
+			}
+
+			//NOTE! These 2 functions are not virtual as we call them in the construct/deconstructor!
+			void InitializeBase();
+			void ShutdownBase();
+
+			//Starts Bundle Persistent Stat Tracking for the given BundleName. Will automatically load applicable data from the disk and reset that cache
+			//if the ExpectedAnalyticsID doesn't match the one previously serialized.
+			//Uses FInstallBundleStatsContainerBase::GetBaseExpectedAnalyticsID if an ExpectedAnalyticsID is not passed in or empty
+			virtual void StartBundlePersistentStatTracking(FName BundleName, const FString& ExpectedAnalyticsID = FString(), bool bForceResetStatData = false);
+			
+			//Starts Session Persistent Stat Tracking under the given SessionName. Will automatically load applicable data from the disk and reset that cache
+			//if the ExpectedAnalyticsID doesn't match the one previously serialized. 
+			//Uses FInstallBundleStatsContainerBase::GetBaseExpectedAnalyticsID if an ExpectedAnalyticsID is not passed in or empty
+			virtual void StartSessionPersistentStatTracking(const FString& SessionName, const TArray<FName>& RequiredBundles = TArray<FName>(), const FString& ExpectedAnalyticsID = FString(), bool bForceResetStatData = false);
+
+			virtual void StopBundlePersistentStatTracking(FName BundleName, bool bStopAllActiveTimers = true);
+			virtual void StopSessionPersistentStatTracking(FString SessionName, bool bStopAllActiveTimers = true);
+
+			virtual void StartBundlePersistentStatTimer(FName BundleName, ETimingStatNames TimerToStart);
+			virtual void StartSessionPersistentStatTimer(FString SessionName, ETimingStatNames TimerToStart);
+
+			virtual void StopBundlePersistentStatTimer(FName BundleName, ETimingStatNames TimerToStop);
+			virtual void StopSessionPersistentStatTimer(FString SessionName, ETimingStatNames TimerToStop);
+
+			virtual void UpdateBundlePersistentStatTimer(FName BundleName, ETimingStatNames TimerToUpdate);
+			virtual void UpdateSessionPersistentStatTimer(FString SessionName, ETimingStatNames TimerToUpdate);
+
+			virtual void IncrementBundlePersistentCounter(FName BundleName, ECountStatNames CounterToUpdate);
+			virtual void IncrementSessionPersistentCounter(FString SessionName, ECountStatNames CounterToUpdate);
+
+			virtual const FBundlePersistentStats* GetBundleStat(FName BundleName) const;
+			virtual const FSessionPersistentStats* GetSessionStat(const FString& SessionName) const;
+
+			virtual void SaveAllDirtyStatsToDisk();
+
+		protected:
+			virtual bool Tick(float dt);
+			virtual void ResetTimerUpdate();
+
+			virtual void UpdateAllBundlesActiveTimers();
+			virtual void UpdateAllSessionActiveTimers();
+
+			virtual void OnApp_EnteringBackground();
+			virtual void OnApp_EnteringForeground();
+
+			virtual void OnBackground_HandleBundleStats();
+			virtual void OnForeground_HandleBundleStats();
+			virtual void OnBackground_HandleSessionStats();
+			virtual void OnForeground_HandleSessionStats();
+			
+			//Called whenever we update our stat's data in one of this class' methods
+			virtual void OnDataUpdatedForStat(FPersistentStatsBase& UpdatedBundleStat);
+
+			//Stops active foreground timers when going to background and starts applicable background version
+			//Also increments background counters
+			virtual void UpdateStatsForBackground(FPersistentStatsBase& StatToUpdate);
+
+			//Stops active background timers and resumes applicable foreground timers
+			//Also increments foreground counters
+			virtual void UpdateStatsForForeground(FPersistentStatsBase& StatToUpdate);
+
+		protected:
+			TMap<FName, FBundlePersistentStats> PerBundlePersistentStatMap;
+			TMap<FString, FSessionPersistentStats> SessionPersistentStatMap;
+
+			FDelegateHandle TickHandle;
+			FDelegateHandle OnApp_EnteringForegroundHandle;
+			FDelegateHandle OnApp_EnteringBackgroundHandle;
+
+			float TimerStat_ResetTimerValue;
+			float TimerStat_UpdateTimer;
+			bool bShouldAutoUpdateInTick;
+			bool bShouldAutoUpdateBackgroundStats;
+			bool bShouldSaveDirtyStatsOnTick;
+			bool bShouldSaveStatsEveryUpdate;
+		};
+	}
+
 	struct INSTALLBUNDLEMANAGER_API IBundleSourceContentRequestSharedContext
 	{
 		virtual ~IBundleSourceContentRequestSharedContext() {}
