Adding system to track persistent install stats between launches of the app. Re-submit of previous CL: 10784102 with fix for Editor build failures and some other improvements.
#tests iOS Client

diff --git a/Engine/Source/Runtime/InstallBundleManager/InstallBundleManager.Build.cs b/Engine/Source/Runtime/InstallBundleManager/InstallBundleManager.Build.cs
--- a/Engine/Source/Runtime/InstallBundleManager/InstallBundleManager.Build.cs
+++ b/Engine/Source/Runtime/InstallBundleManager/InstallBundleManager.Build.cs
@@ -16,6 +16,7 @@ public class InstallBundleManager : ModuleRules
 				"InputCore",
 				"Projects",
 				"ApplicationCore",
+				"Json"
 			}
 		);
 	}
diff --git a/Engine/Source/Runtime/InstallBundleManager/Private/InstallBundleUtils.cpp b/Engine/Source/Runtime/InstallBundleManager/Private/InstallBundleUtils.cpp
--- a/Engine/Source/Runtime/InstallBundleManager/Private/InstallBundleUtils.cpp
+++ b/Engine/Source/Runtime/InstallBundleManager/Private/InstallBundleUtils.cpp
@@ -6,6 +6,15 @@
 
 #include "HAL/PlatformApplicationMisc.h"
 
+#include "Containers/Ticker.h"
+
+#include "Misc/ConfigCacheIni.h"
+#include "Misc/CoreDelegates.h"
+#include "Misc/Paths.h"
+#include "Misc/FileHelper.h"
+#include "HAL/PlatformFilemanager.h"
+#include "Serialization/JsonSerializerMacros.h"
+
 namespace InstallBundleUtil
 {
 	FString GetAppVersion()
@@ -160,4 +169,1056 @@ namespace InstallBundleUtil
 			StateStats.bOpen = false;
 		}
 	}
+
+	namespace PersistentStats
+	{
+		const FString& LexToString(ETimingStatNames InType)
+		{
+			static const FString TotalTime_Real(TEXT("TotalTime_Real"));
+			static const FString TotalTime_FG(TEXT("TotalTime_FG"));
+			static const FString TotalTime_BG(TEXT("TotalTime_BG"));
+			static const FString ChunkDBDownloadTime_Real(TEXT("ChunkDBDownloadTime_Real"));
+			static const FString ChunkDBDownloadTime_FG(TEXT("ChunkDBDownloadTime_FG"));
+			static const FString ChunkDBDownloadTime_BG(TEXT("ChunkDBDownloadTime_BG"));
+			static const FString InstallTime_Real(TEXT("InstallTime_Real"));
+			static const FString InstallTime_FG(TEXT("InstallTime_FG"));
+			static const FString InstallTime_BG(TEXT("InstallTime_BG"));
+			static const FString PSOTime_Real(TEXT("PSOTime_Real"));
+			static const FString PSOTime_FG(TEXT("PSOTime_FG"));
+			static const FString PSOTime_BG(TEXT("PSOTime_BG"));
+
+			static const FString Unknown(TEXT("<Unknown PersistentStats::ETimingStatNames Value>"));
+
+			switch (InType)
+			{
+			case ETimingStatNames::TotalTime_Real:
+				return TotalTime_Real;
+			case ETimingStatNames::TotalTime_FG:
+				return TotalTime_FG;
+			case ETimingStatNames::TotalTime_BG:
+				return TotalTime_BG;
+			case ETimingStatNames::ChunkDBDownloadTime_Real:
+				return ChunkDBDownloadTime_Real;
+			case ETimingStatNames::ChunkDBDownloadTime_FG:
+				return ChunkDBDownloadTime_FG;
+			case ETimingStatNames::ChunkDBDownloadTime_BG:
+				return ChunkDBDownloadTime_BG;
+			case ETimingStatNames::InstallTime_Real:
+				return InstallTime_Real;
+			case ETimingStatNames::InstallTime_FG:
+				return InstallTime_FG;
+			case ETimingStatNames::InstallTime_BG:
+				return InstallTime_BG;
+			case ETimingStatNames::PSOTime_Real:
+				return PSOTime_Real;
+			case ETimingStatNames::PSOTime_FG:
+				return PSOTime_FG;
+			case ETimingStatNames::PSOTime_BG:
+				return PSOTime_BG;
+			default:
+				break;
+			}
+
+			ensureAlwaysMsgf(false, TEXT("Missing PersistentStats::ETimingStatNames LexToString entry! Missing Entry as Int: %d"), (int)(InType));
+			return Unknown;
+		}
+	
+		const FString& LexToString(ECountStatNames InType)
+		{
+			static const FString NumResumedFromBackground(TEXT("NumResumedFromBackground"));
+			static const FString NumResumedFromLaunch(TEXT("NumResumedFromLaunch"));
+			static const FString NumBackgrounded(TEXT("NumBackgrounded"));
+
+			static const FString Unknown(TEXT("<Unknown PersistentStats::ETimingStatNames Value>"));
+
+			switch (InType)
+			{
+			case ECountStatNames::NumResumedFromBackground:
+				return NumResumedFromBackground;
+			case ECountStatNames::NumResumedFromLaunch:
+				return NumResumedFromLaunch;
+			case ECountStatNames::NumBackgrounded:
+				return NumBackgrounded;
+			default:
+				break;
+			}
+
+			ensureAlwaysMsgf(false, TEXT("Missing PersistentStats::ECountStatNames LexToString entry! Missing Entry as Int: %d"), (int)(InType));
+			return Unknown;
+		}
+
+		
+		bool IsTimerReal(ETimingStatNames InTimerType)
+		{
+			switch (InTimerType)
+			{
+				//Intentional fallthrough for known true types
+				case ETimingStatNames::TotalTime_Real:
+				case ETimingStatNames::ChunkDBDownloadTime_Real:
+				case ETimingStatNames::InstallTime_Real:
+				case ETimingStatNames::PSOTime_Real:
+					return true;
+
+				//Intentional fallthrough for known false types
+				case ETimingStatNames::TotalTime_FG:
+				case ETimingStatNames::TotalTime_BG:
+				case ETimingStatNames::ChunkDBDownloadTime_FG:
+				case ETimingStatNames::ChunkDBDownloadTime_BG:
+				case ETimingStatNames::InstallTime_FG:
+				case ETimingStatNames::InstallTime_BG:
+				case ETimingStatNames::PSOTime_FG:
+				case ETimingStatNames::PSOTime_BG:
+					return false;
+					
+				default:
+					break;
+			}
+
+			ensureAlwaysMsgf(false, TEXT("Missing PersistentStats::ETimingStatNames IsTimerReal entry! Missing Entry:%s"), *LexToString(InTimerType));
+			return false;
+		}
+
+		bool IsTimerFG(ETimingStatNames InTimerType)
+		{
+			switch (InTimerType)
+			{
+				//Intentional fallthrough for known true types
+				case ETimingStatNames::TotalTime_FG:
+				case ETimingStatNames::ChunkDBDownloadTime_FG:
+				case ETimingStatNames::InstallTime_FG:
+				case ETimingStatNames::PSOTime_FG:
+					return true;
+
+				//Intentional fallthrough for known false types
+				case ETimingStatNames::TotalTime_Real:
+				case ETimingStatNames::TotalTime_BG:
+				case ETimingStatNames::ChunkDBDownloadTime_Real:
+				case ETimingStatNames::ChunkDBDownloadTime_BG:
+				case ETimingStatNames::InstallTime_Real:
+				case ETimingStatNames::InstallTime_BG:
+				case ETimingStatNames::PSOTime_Real:
+				case ETimingStatNames::PSOTime_BG:
+					return false;
+
+				default:
+					break;
+			}
+
+			ensureAlwaysMsgf(false, TEXT("Missing PersistentStats::ETimingStatNames IsTimerFG entry! Missing Entry as Int: %s"), *LexToString(InTimerType));
+			return false;
+		}
+
+		bool IsTimerBG(ETimingStatNames InTimerType)
+		{
+			switch (InTimerType)
+			{
+				//Intentional fallthrough for known true types
+				case ETimingStatNames::TotalTime_BG:
+				case ETimingStatNames::ChunkDBDownloadTime_BG:
+				case ETimingStatNames::InstallTime_BG:
+				case ETimingStatNames::PSOTime_BG:
+					return true;
+
+				//Intentional fallthrough for known false types
+				case ETimingStatNames::TotalTime_Real:
+				case ETimingStatNames::TotalTime_FG:
+				case ETimingStatNames::ChunkDBDownloadTime_Real:
+				case ETimingStatNames::ChunkDBDownloadTime_FG:
+				case ETimingStatNames::InstallTime_Real:
+				case ETimingStatNames::InstallTime_FG:
+				case ETimingStatNames::PSOTime_Real:
+				case ETimingStatNames::PSOTime_FG:
+					return false;
+
+				default:
+					break;
+			}
+
+			ensureAlwaysMsgf(false, TEXT("Missing PersistentStats::ETimingStatNames IsTimerBG entry! Missing Entry as Int: %s"), *LexToString(InTimerType));
+			return false;
+		}
+
+		ETimingStatNames GetAssociatedRealTimerName(ETimingStatNames InTimerType)
+		{
+			switch (InTimerType)
+			{
+				//Intentional fallthrough for TotalTime
+				case ETimingStatNames::TotalTime_Real:
+				case ETimingStatNames::TotalTime_FG:
+				case ETimingStatNames::TotalTime_BG:
+					return ETimingStatNames::TotalTime_Real;
+
+				//Intentional fallthrough for ChunkDBDownloadTime
+				case ETimingStatNames::ChunkDBDownloadTime_Real:
+				case ETimingStatNames::ChunkDBDownloadTime_FG:
+				case ETimingStatNames::ChunkDBDownloadTime_BG:
+					return ETimingStatNames::ChunkDBDownloadTime_Real;
+
+				//Intentional fallthrough for InstallTime
+				case ETimingStatNames::InstallTime_Real:
+				case ETimingStatNames::InstallTime_FG:
+				case ETimingStatNames::InstallTime_BG:
+					return ETimingStatNames::InstallTime_Real;
+
+				//Intentional fallthrough for PSOTime
+				case ETimingStatNames::PSOTime_Real:
+				case ETimingStatNames::PSOTime_FG:
+				case ETimingStatNames::PSOTime_BG:
+					return ETimingStatNames::PSOTime_Real;
+
+				default:
+					break;
+			}
+
+			ensureAlwaysMsgf(false, TEXT("Missing PersistentStats::ETimingStatNames GetAssociatedRealTimerName entry! Missing Entry as Int: %s"), *LexToString(InTimerType));
+			return ETimingStatNames::NumStatNames;
+		}
+
+		ETimingStatNames GetAssociatedFGTimerName(ETimingStatNames InTimerType)
+		{
+			switch (InTimerType)
+			{
+				//Intentional fallthrough for TotalTime
+				case ETimingStatNames::TotalTime_Real:
+				case ETimingStatNames::TotalTime_FG:
+				case ETimingStatNames::TotalTime_BG:
+					return ETimingStatNames::TotalTime_FG;
+
+				//Intentional fallthrough for ChunkDBDownloadTime
+				case ETimingStatNames::ChunkDBDownloadTime_Real:
+				case ETimingStatNames::ChunkDBDownloadTime_FG:
+				case ETimingStatNames::ChunkDBDownloadTime_BG:
+					return ETimingStatNames::ChunkDBDownloadTime_FG;
+
+				//Intentional fallthrough for InstallTime
+				case ETimingStatNames::InstallTime_Real:
+				case ETimingStatNames::InstallTime_FG:
+				case ETimingStatNames::InstallTime_BG:
+					return ETimingStatNames::InstallTime_FG;
+
+				//Intentional fallthrough for PSOTime
+				case ETimingStatNames::PSOTime_Real:
+				case ETimingStatNames::PSOTime_FG:
+				case ETimingStatNames::PSOTime_BG:
+					return ETimingStatNames::PSOTime_FG;
+
+				default:
+					break;
+			}
+
+			ensureAlwaysMsgf(false, TEXT("Missing PersistentStats::ETimingStatNames GetAssociatedFGTimerName entry! Missing Entry as Int: %s"), *LexToString(InTimerType));
+			return ETimingStatNames::NumStatNames;
+		}
+
+		ETimingStatNames GetAssociatedBGTimerName(ETimingStatNames InTimerType)
+		{
+			switch (InTimerType)
+			{
+				//Intentional fallthrough for TotalTime
+				case ETimingStatNames::TotalTime_Real:
+				case ETimingStatNames::TotalTime_FG:
+				case ETimingStatNames::TotalTime_BG:
+					return ETimingStatNames::TotalTime_BG;
+
+				//Intentional fallthrough for ChunkDBDownloadTime
+				case ETimingStatNames::ChunkDBDownloadTime_Real:
+				case ETimingStatNames::ChunkDBDownloadTime_FG:
+				case ETimingStatNames::ChunkDBDownloadTime_BG:
+					return ETimingStatNames::ChunkDBDownloadTime_BG;
+
+				//Intentional fallthrough for InstallTime
+				case ETimingStatNames::InstallTime_Real:
+				case ETimingStatNames::InstallTime_FG:
+				case ETimingStatNames::InstallTime_BG:
+					return ETimingStatNames::InstallTime_BG;
+
+				//Intentional fallthrough for PSOTime
+				case ETimingStatNames::PSOTime_Real:
+				case ETimingStatNames::PSOTime_FG:
+				case ETimingStatNames::PSOTime_BG:
+					return ETimingStatNames::PSOTime_BG;
+
+			default:
+				break;
+			}
+
+			ensureAlwaysMsgf(false, TEXT("Missing PersistentStats::ETimingStatNames GetAssociatedBGTimerName entry! Missing Entry as Int: %s"), *LexToString(InTimerType));
+			return ETimingStatNames::NumStatNames;
+		}
+
+		bool FPersistentStatsBase::LoadStatsFromDisk()
+		{
+			FString JSONStringOnDisk;
+			if (FPaths::FileExists(GetFullPathForStatFile()))
+			{
+				FFileHelper::LoadFileToString(JSONStringOnDisk, *GetFullPathForStatFile());
+			}
+
+			if (!JSONStringOnDisk.IsEmpty())
+			{
+				return FromJson(JSONStringOnDisk);
+			}
+
+			return false;
+		}
+
+		bool FPersistentStatsBase::SaveStatsToDisk()
+		{
+			bIsDirty = false;
+			return FFileHelper::SaveStringToFile(ToJson(), *GetFullPathForStatFile());
+		}
+
+		void FPersistentStatsBase::ResetStats(const FString& NewAnalyticsSessionID)
+		{
+			TimingStatsMap.Reset();
+			CountStatMap.Reset();
+			AnalyticsSessionID = NewAnalyticsSessionID;
+
+			bIsDirty = true;
+		}
+
+		bool FPersistentStatsBase::HasTimingStat(ETimingStatNames StatToCheck) const
+		{
+			const FPersistentTimerData* FoundStat = TimingStatsMap.Find(LexToString(StatToCheck));
+			return (nullptr != FoundStat);
+		}
+
+		bool FPersistentStatsBase::HasCountStat(ECountStatNames StatToCheck) const
+		{
+			const int* FoundStat = CountStatMap.Find(LexToString(StatToCheck));
+			return (nullptr != FoundStat);
+		}
+
+		const FPersistentTimerData* FPersistentStatsBase::GetTimingStatData(ETimingStatNames StatToGet) const
+		{ 
+			return TimingStatsMap.Find(LexToString(StatToGet));
+		}
+
+		const int* FPersistentStatsBase::GetCountStatData(ECountStatNames StatToGet) const
+		{
+			return CountStatMap.Find(LexToString(StatToGet));
+		}
+
+		void FPersistentStatsBase::IncrementCountStat(PersistentStats::ECountStatNames StatToUpdate)
+		{
+			int& StatCount = CountStatMap.FindOrAdd(LexToString(StatToUpdate));
+			++StatCount;
+
+			bIsDirty = true;
+		}
+
+		bool FPersistentStatsBase::IsTimingStatStarted(PersistentStats::ETimingStatNames StatToUpdate) const
+		{
+			bool HasStarted = false;
+
+			if (HasTimingStat(StatToUpdate))
+			{
+				const FPersistentTimerData* FoundStat = GetTimingStatData(StatToUpdate);
+				if (ensureAlwaysMsgf((nullptr != FoundStat), TEXT("Missing FInstallBundlePersistentTimingData but returned true from HasTimingStat For Stat:%s!"), *LexToString(StatToUpdate)))
+				{
+					HasStarted = (FoundStat->LastUpdateTime != 0.);
+				}
+			}
+
+			return HasStarted;
+		}
+
+		void FPersistentStatsBase::StartTimingStat(PersistentStats::ETimingStatNames StatToUpdate)
+		{
+			if (!IsTimingStatStarted(StatToUpdate))
+			{
+				FPersistentTimerData& FoundStat = TimingStatsMap.FindOrAdd(LexToString(StatToUpdate));
+				FoundStat.LastUpdateTime = FPlatformTime::Seconds();
+			}
+			//if this stat was already updated then instead of losing the time since its last update by
+			//starting it again lets just update it to keep that time
+			else
+			{
+				UpdateTimingStat(StatToUpdate);
+			}
+			bIsDirty = true;
+		}
+
+		void FPersistentStatsBase::StopTimingStat(PersistentStats::ETimingStatNames StatToUpdate, bool UpdateTimerOnStop /* = true */)
+		{
+			//Only want to actually update the timer if we have started it (otherwise the update won't do anything and will ensure)
+			if (UpdateTimerOnStop && IsTimingStatStarted(StatToUpdate))
+			{
+				UpdateTimingStat(StatToUpdate);
+			}
+
+			FPersistentTimerData& FoundStat = TimingStatsMap.FindOrAdd(LexToString(StatToUpdate));
+			FoundStat.LastUpdateTime = 0.;
+
+			bIsDirty = true;
+		}
+
+		void FPersistentStatsBase::UpdateTimingStat(PersistentStats::ETimingStatNames StatToUpdate)
+		{
+			if (ensureAlwaysMsgf(IsTimingStatStarted(StatToUpdate), TEXT("Calling UpdateTimingStat on a stat that hasn't been started! %s"), *LexToString(StatToUpdate)))
+			{
+				FPersistentTimerData& FoundStat = TimingStatsMap.FindOrAdd(LexToString(StatToUpdate));
+
+				const double CurrentTime = FPlatformTime::Seconds();
+				const double TimeSinceUpdate = CurrentTime - FoundStat.LastUpdateTime;
+
+				if (ensureAlwaysMsgf((TimeSinceUpdate > 0.f), TEXT("Logic Error! Invalid saved LastUpdateTime for Stat %s!"), *LexToString(StatToUpdate)))
+				{
+					FoundStat.CurrentValue += TimeSinceUpdate;
+				}
+				
+				FoundStat.LastUpdateTime = CurrentTime;
+				bIsDirty = true;
+			}
+		}
+
+		void FPersistentStatsBase::UpdateAllActiveTimers()
+		{
+			for (uint8 TimingStatNameIndex = 0; TimingStatNameIndex < (uint8)PersistentStats::ETimingStatNames::NumStatNames; ++TimingStatNameIndex)
+			{
+				PersistentStats::ETimingStatNames EnumForIndex = (PersistentStats::ETimingStatNames)TimingStatNameIndex;
+				if (IsTimingStatStarted(EnumForIndex))
+				{
+					UpdateTimingStat(EnumForIndex);
+				}
+			}
+		}
+
+		void FPersistentStatsBase::StopAllActiveTimers()
+		{
+			for (uint8 TimingStatIndex = 0; TimingStatIndex < static_cast<uint8>(ETimingStatNames::NumStatNames); ++TimingStatIndex)
+			{
+				ETimingStatNames TimingStatAsEnum = static_cast<ETimingStatNames>(TimingStatIndex);
+				if (IsTimingStatStarted(TimingStatAsEnum))
+				{
+					StopTimingStat(TimingStatAsEnum);
+				}
+			}
+		}
+
+		void FPersistentStatsBase::StatsBegin(const FString& ExpectedAnalyticsID, bool bForceResetData /* = false */)
+		{
+			bIsActive = true;
+			
+			if (LoadStatsFromDisk())
+			{
+				OnLoadingDataFromDisk();
+			}
+
+			//If our Analytics ID doesn't match our expected we need to reset the data as we have started a new persistent session
+			if (bForceResetData || !AnalyticsSessionID.Equals(ExpectedAnalyticsID))
+			{
+				ResetStats(ExpectedAnalyticsID);
+			}
+			
+			//Immediately save here so we don't risk reloading the same stale data
+			//if we don't make it to an update
+			SaveStatsToDisk();
+		}
+
+		void FPersistentStatsBase::StatsEnd(bool bStopAllActiveTimers /* = true */)
+		{
+			bIsActive = false;
+
+			if (bStopAllActiveTimers)
+			{
+				StopAllActiveTimers();
+			}
+
+			//Immediately save here as we only look to update active dirty bundles, and since
+			//this likely won't be changed anymore we might as well save it out now
+			SaveStatsToDisk();
+		}
+
+		void FPersistentStatsBase::OnLoadingDataFromDisk()
+		{
+			HandleTimerStatsAfterDataLoad();
+		}
+
+		void FPersistentStatsBase::HandleTimerStatsAfterDataLoad()
+		{
+			//Go through all timing stats and handle each one accordingly
+			//All Real timers should be updated after load without being stopped.
+			//All FG Timers we should stop these timers without updating them so that they don't accrue time from backgrounding
+			//All BG timers Should be stopped, but update their timers on stopping so they accrue time from being inactive
+
+			/*
+					Handle Real Timers
+			*/
+
+			if (IsTimingStatStarted(ETimingStatNames::TotalTime_Real))
+			{
+				UpdateTimingStat(ETimingStatNames::TotalTime_Real);
+			}
+
+			if (IsTimingStatStarted(ETimingStatNames::ChunkDBDownloadTime_Real))
+			{
+				UpdateTimingStat(ETimingStatNames::ChunkDBDownloadTime_Real);
+			}
+
+			if (IsTimingStatStarted(ETimingStatNames::InstallTime_Real))
+			{
+				UpdateTimingStat(ETimingStatNames::InstallTime_Real);
+			}
+
+			if (IsTimingStatStarted(ETimingStatNames::PSOTime_Real))
+			{
+				UpdateTimingStat(ETimingStatNames::PSOTime_Real);
+			}
+
+			/*
+					Handle Foreground Timers
+			*/
+
+			if (IsTimingStatStarted(ETimingStatNames::TotalTime_FG))
+			{
+				StopTimingStat(ETimingStatNames::TotalTime_FG, false);
+			}
+
+			if (IsTimingStatStarted(ETimingStatNames::ChunkDBDownloadTime_FG))
+			{
+				StopTimingStat(ETimingStatNames::ChunkDBDownloadTime_FG, false);
+			}
+
+			if (IsTimingStatStarted(ETimingStatNames::InstallTime_FG))
+			{
+				StopTimingStat(ETimingStatNames::InstallTime_FG, false);
+			}
+
+			if (IsTimingStatStarted(ETimingStatNames::PSOTime_FG))
+			{
+				StopTimingStat(ETimingStatNames::PSOTime_FG, false);
+			}
+
+			/*
+					Handle Background Timers
+			*/
+			
+			if (IsTimingStatStarted(ETimingStatNames::TotalTime_BG))
+			{
+				StopTimingStat(ETimingStatNames::TotalTime_BG, true);
+			}
+
+			if (IsTimingStatStarted(ETimingStatNames::ChunkDBDownloadTime_BG))
+			{
+				StopTimingStat(ETimingStatNames::ChunkDBDownloadTime_BG, true);
+			}
+
+			if (IsTimingStatStarted(ETimingStatNames::InstallTime_BG))
+			{
+				StopTimingStat(ETimingStatNames::InstallTime_BG, true);
+			}
+
+			if (IsTimingStatStarted(ETimingStatNames::PSOTime_BG))
+			{
+				StopTimingStat(ETimingStatNames::PSOTime_BG, true);
+			}
+		}
+
+		void FSessionPersistentStats::AddRequiredBundles(const TArray<FString>& RequiredBundlesToAdd)
+		{
+			for (const FString& BundleName : RequiredBundlesToAdd)
+			{
+				RequiredBundles.AddUnique(BundleName);
+			}
+			bIsDirty = true;
+		}
+
+		void FSessionPersistentStats::AddRequiredBundles(const TArray<FName>& RequiredBundlesToAdd)
+		{
+			for (FName BundleName : RequiredBundlesToAdd)
+			{
+				RequiredBundles.AddUnique(BundleName.ToString());
+			}
+			bIsDirty = true;
+		}
+
+		void FSessionPersistentStats::GetRequiredBundles(TArray<FString>& OutRequiredBundles) const
+		{
+			OutRequiredBundles.Empty();
+			for (const FString& BundleName : RequiredBundles)
+			{
+				OutRequiredBundles.Add(BundleName);
+			}
+		}
+	
+		void FSessionPersistentStats::ResetRequiredBundles(const TArray<FString>& NewRequiredBundles /* = TArray<FString>() */)
+		{
+			RequiredBundles.Empty();
+			AddRequiredBundles(NewRequiredBundles);
+			bIsDirty = true;
+		}
+
+		const FString FBundlePersistentStats::GetFullPathForStatFile() const
+		{
+			return FPaths::Combine(FPlatformMisc::GamePersistentDownloadDir(), TEXT("PersistentStats"), TEXT("BundleStats"), (BundleName + TEXT(".json")));
+		}
+
+		const FString FSessionPersistentStats::GetFullPathForStatFile() const
+		{
+			return FPaths::Combine(FPlatformMisc::GamePersistentDownloadDir(), TEXT("PersistentStats"), TEXT("ContentRequestStats"), (SessionName + TEXT(".json")));
+		}
+
+		FPersistentStatContainerBase::FPersistentStatContainerBase()
+			: PerBundlePersistentStatMap()
+			, SessionPersistentStatMap()
+			, TickHandle()
+			, OnApp_EnteringForegroundHandle()
+			, OnApp_EnteringBackgroundHandle()
+			, TimerAutoUpdateTimeRemaining(10.0f)
+			, TimerDirtyStatUpdateTimeRemaining(10.0f)
+			, bShouldAutoUpdateTimersInTick(true)
+			, TimerAutoUpdateRate(10.0f)
+			, bShouldSaveDirtyStatsOnTick(true)
+			, DirtyStatSaveToDiskRate(5.f)
+			, bShouldSaveStatsOnUpdate(false)
+			, bShouldAutoHandleFGBGStats(true)
+		{
+			InitializeBase();
+		}
+
+		FPersistentStatContainerBase::~FPersistentStatContainerBase()
+		{
+			ShutdownBase();
+		}
+
+		void FPersistentStatContainerBase::InitializeBase()
+		{
+			//Load Settings from Config
+			{
+				GConfig->GetBool(TEXT("InstallBundleManager.PersistentStatSettings"), TEXT("bShouldAutoUpdateTimersInTick"), bShouldAutoUpdateTimersInTick, GEngineIni);
+				GConfig->GetFloat(TEXT("InstallBundleManager.PersistentStatSettings"), TEXT("TimerAutoUpdateRate"), TimerAutoUpdateRate, GEngineIni);
+
+				GConfig->GetBool(TEXT("InstallBundleManager.PersistentStatSettings"), TEXT("bShouldSaveDirtyStatsOnTick"), bShouldSaveDirtyStatsOnTick, GEngineIni);
+				GConfig->GetFloat(TEXT("InstallBundleManager.PersistentStatSettings"), TEXT("DirtyStatSaveToDiskRate"), DirtyStatSaveToDiskRate, GEngineIni);
+				GConfig->GetBool(TEXT("InstallBundleManager.PersistentStatSettings"), TEXT("bShouldSaveStatsOnUpdate"), bShouldSaveStatsOnUpdate, GEngineIni);
+				
+				GConfig->GetBool(TEXT("InstallBundleManager.PersistentStatSettings"), TEXT("bShouldAutoHandleFGBGStats"), bShouldAutoHandleFGBGStats, GEngineIni);
+
+				//Reset timers so they follow the new loaded-in value
+				ResetTimerUpdate();
+				ResetDirtyStatUpdate();
+			}
+
+			//Setup Delegates (Needs to happen after config to have AutoUpdate settings loaded)
+			{
+				//Only setup a tick function if we would use it
+				if ((bShouldAutoUpdateTimersInTick || bShouldSaveDirtyStatsOnTick) && !TickHandle.IsValid())
+				{
+					TickHandle = FTicker::GetCoreTicker().AddTicker(FTickerDelegate::CreateRaw(this, &FPersistentStatContainerBase::Tick));
+				}
+
+				//Only setup Foreground/Background delegates if we should be using them to swap stats
+				if (bShouldAutoHandleFGBGStats)
+				{
+					if (!OnApp_EnteringForegroundHandle.IsValid())
+					{
+						OnApp_EnteringForegroundHandle = FCoreDelegates::ApplicationHasEnteredForegroundDelegate.AddRaw(this, &FPersistentStatContainerBase::OnApp_EnteringForeground);
+					}
+					if (!OnApp_EnteringBackgroundHandle.IsValid())
+					{
+						OnApp_EnteringBackgroundHandle = FCoreDelegates::ApplicationWillEnterBackgroundDelegate.AddRaw(this, &FPersistentStatContainerBase::OnApp_EnteringBackground);
+					}
+				}
+			}
+		}
+
+		void FPersistentStatContainerBase::ShutdownBase()
+		{
+			if (TickHandle.IsValid())
+			{
+				FTicker::GetCoreTicker().RemoveTicker(TickHandle);
+				TickHandle.Reset();
+			}
+
+			if (OnApp_EnteringForegroundHandle.IsValid())
+			{
+				FCoreDelegates::ApplicationHasEnteredForegroundDelegate.Remove(OnApp_EnteringForegroundHandle);
+				OnApp_EnteringForegroundHandle.Reset();
+			}
+			if (OnApp_EnteringBackgroundHandle.IsValid())
+			{
+				FCoreDelegates::ApplicationWillEnterBackgroundDelegate.Remove(OnApp_EnteringBackgroundHandle);
+				OnApp_EnteringBackgroundHandle.Reset();
+			}
+		}
+
+		bool FPersistentStatContainerBase::Tick(float dt)
+		{
+			if (bShouldAutoUpdateTimersInTick)
+			{
+				//Only update all active timers every TimerStat_ResetTimerValue seconds
+				TimerAutoUpdateTimeRemaining -= dt;
+				if (TimerAutoUpdateTimeRemaining < 0.0f)
+				{
+					ResetTimerUpdate();
+					UpdateAllBundlesActiveTimers();
+					UpdateAllSessionActiveTimers();
+				}
+			}
+
+			if (bShouldSaveDirtyStatsOnTick)
+			{
+				TimerDirtyStatUpdateTimeRemaining -= dt;
+				if (TimerDirtyStatUpdateTimeRemaining < 0.0f)
+				{
+					ResetDirtyStatUpdate();
+					SaveAllDirtyStatsToDisk();
+				}
+			}
+
+			//go ahead and always Tick once we start
+			return true;
+		}
+
+		void FPersistentStatContainerBase::OnTimerStartedForStat(FPersistentStatsBase& BundleStatForTimer, ETimingStatNames TimerStarted)
+		{
+			//If we are auto handling FG/BG stats and we have a real timer, start the _FG version with the _Real version
+			if (bShouldAutoHandleFGBGStats && IsTimerReal(TimerStarted))
+			{
+				BundleStatForTimer.StartTimingStat(GetAssociatedFGTimerName(TimerStarted));
+
+				//We should also check to make sure the _BG version is stopped if it is running
+				const ETimingStatNames BGTimingStat = GetAssociatedBGTimerName(TimerStarted);
+				if (BundleStatForTimer.IsTimingStatStarted(BGTimingStat))
+				{
+					BundleStatForTimer.StopTimingStat(BGTimingStat, true);
+				}
+			}
+		}
+
+		void FPersistentStatContainerBase::OnTimerStoppedForStat(FPersistentStatsBase& BundleStatForTimer, ETimingStatNames TimerStarted)
+		{
+			//If we are auto handling FG/BG stats and we have a real timer, stop the _FG and _BG version with the _Real version
+			if (bShouldAutoHandleFGBGStats && IsTimerReal(TimerStarted))
+			{
+				BundleStatForTimer.StopTimingStat(GetAssociatedFGTimerName(TimerStarted));
+				BundleStatForTimer.StopTimingStat(GetAssociatedBGTimerName(TimerStarted));
+			}
+		}
+
+		void FPersistentStatContainerBase::OnDataUpdatedForStat(FPersistentStatsBase& UpdatedBundleStat)
+		{
+			if (bShouldSaveStatsOnUpdate)
+			{
+				UpdatedBundleStat.SaveStatsToDisk();
+			}
+		}
+
+		void FPersistentStatContainerBase::ResetTimerUpdate()
+		{
+			TimerAutoUpdateTimeRemaining = TimerAutoUpdateRate;
+		}
+
+		void FPersistentStatContainerBase::ResetDirtyStatUpdate()
+		{
+			TimerDirtyStatUpdateTimeRemaining = DirtyStatSaveToDiskRate;
+		}		
+
+		void FPersistentStatContainerBase::SaveAllDirtyStatsToDisk()
+		{
+			TArray<FName> AllBundleStatNames;
+			PerBundlePersistentStatMap.GetKeys(AllBundleStatNames);
+			for (FName& BundleName : AllBundleStatNames)
+			{
+				FBundlePersistentStats* BundleStats = PerBundlePersistentStatMap.Find(BundleName);
+				check(BundleStats);
+				if (BundleStats->IsDirty())
+				{
+					BundleStats->SaveStatsToDisk();
+				}
+			}
+
+			TArray<FString> AllSessionStatNames;
+			SessionPersistentStatMap.GetKeys(AllSessionStatNames);
+			for (const FString& SessionName : AllSessionStatNames)
+			{
+				FSessionPersistentStats* SessionStats = SessionPersistentStatMap.Find(SessionName);
+				check(SessionStats);
+				if (SessionStats->IsDirty())
+				{
+					SessionStats->SaveStatsToDisk();
+				}
+			}			
+		}
+
+		void FPersistentStatContainerBase::StartBundlePersistentStatTracking(FName BundleName, const FString& ExpectedAnalyticsID /* = FString() */, bool bForceResetStatData /* = false */)
+		{
+			//Use the base expected analytics ID if one was not passed in
+			const FString ExpectedAnalyticsToUse = ExpectedAnalyticsID.IsEmpty() ? FPersistentStatsBase::GetBaseExpectedAnalyticsID() : ExpectedAnalyticsID;
+
+			FBundlePersistentStats& FoundBundleStats = PerBundlePersistentStatMap.FindOrAdd(BundleName, FBundlePersistentStats(BundleName));
+			FoundBundleStats.StatsBegin(ExpectedAnalyticsToUse, bForceResetStatData);
+
+			OnDataUpdatedForStat(FoundBundleStats);
+		}
+
+		void FPersistentStatContainerBase::StartSessionPersistentStatTracking(const FString& SessionName, const TArray<FName>& RequiredBundles /* = TArray<FName>() */, const FString& ExpectedAnalyticsID /* = FString() */, bool bForceResetStatData /* = false */)
+		{
+			//Use the base expected analytics ID if one was not passed in
+			const FString ExpectedAnalyticsToUse = ExpectedAnalyticsID.IsEmpty() ? FPersistentStatsBase::GetBaseExpectedAnalyticsID() : ExpectedAnalyticsID;
+
+			FSessionPersistentStats& FoundSessionStats = SessionPersistentStatMap.FindOrAdd(SessionName, FSessionPersistentStats(SessionName));
+			FoundSessionStats.StatsBegin(ExpectedAnalyticsToUse, bForceResetStatData);
+
+			//Also append starting required bundles as we may have new ones from the ones already in data
+			FoundSessionStats.AddRequiredBundles(RequiredBundles);
+
+			OnDataUpdatedForStat(FoundSessionStats);
+		}
+		
+		void FPersistentStatContainerBase::StopBundlePersistentStatTracking(FName BundleName, bool bStopAllActiveTimers /* = true */)
+		{
+			FBundlePersistentStats* FoundBundleStats = PerBundlePersistentStatMap.Find(BundleName);
+			if (nullptr != FoundBundleStats)
+			{
+				FoundBundleStats->StatsEnd(bStopAllActiveTimers);
+				OnDataUpdatedForStat(*FoundBundleStats);
+			}
+		}
+
+		void FPersistentStatContainerBase::StopSessionPersistentStatTracking(FString SessionName, bool bStopAllActiveTimers /* = true */)
+		{
+			FSessionPersistentStats* FoundSessionStats = SessionPersistentStatMap.Find(SessionName);
+			if (nullptr != FoundSessionStats)
+			{
+				FoundSessionStats->StatsEnd(bStopAllActiveTimers);
+				OnDataUpdatedForStat(*FoundSessionStats);
+			}
+		}
+
+		void FPersistentStatContainerBase::StartBundlePersistentStatTimer(FName BundleName, ETimingStatNames TimerToStart)
+		{
+			FBundlePersistentStats& FoundBundleStats = PerBundlePersistentStatMap.FindOrAdd(BundleName, FBundlePersistentStats(BundleName));
+			FoundBundleStats.StartTimingStat(TimerToStart);
+
+			OnTimerStartedForStat(FoundBundleStats, TimerToStart);
+			OnDataUpdatedForStat(FoundBundleStats);
+		}
+
+		void FPersistentStatContainerBase::StartSessionPersistentStatTimer(FString SessionName, ETimingStatNames TimerToStart)
+		{
+			FSessionPersistentStats& FoundSessionStats = SessionPersistentStatMap.FindOrAdd(SessionName, FSessionPersistentStats(SessionName));
+			FoundSessionStats.StartTimingStat(TimerToStart);
+
+			OnTimerStartedForStat(FoundSessionStats, TimerToStart);
+			OnDataUpdatedForStat(FoundSessionStats);
+		}
+
+		void FPersistentStatContainerBase::StopBundlePersistentStatTimer(FName BundleName, ETimingStatNames TimerToStop)
+		{
+			FBundlePersistentStats& FoundBundleStats = PerBundlePersistentStatMap.FindOrAdd(BundleName, FBundlePersistentStats(BundleName));
+			FoundBundleStats.StopTimingStat(TimerToStop);
+
+			OnTimerStoppedForStat(FoundBundleStats, TimerToStop);
+			OnDataUpdatedForStat(FoundBundleStats);
+		}
+
+		void FPersistentStatContainerBase::StopSessionPersistentStatTimer(FString SessionName, ETimingStatNames TimerToStop)
+		{
+			FSessionPersistentStats& FoundSessionStats = SessionPersistentStatMap.FindOrAdd(SessionName, FSessionPersistentStats(SessionName));
+			FoundSessionStats.StopTimingStat(TimerToStop);
+
+			OnTimerStoppedForStat(FoundSessionStats, TimerToStop);
+			OnDataUpdatedForStat(FoundSessionStats);
+		}
+
+		void FPersistentStatContainerBase::UpdateBundlePersistentStatTimer(FName BundleName, ETimingStatNames TimerToUpdate)
+		{
+			FBundlePersistentStats& FoundBundleStats = PerBundlePersistentStatMap.FindOrAdd(BundleName, FBundlePersistentStats(BundleName));
+			FoundBundleStats.UpdateTimingStat(TimerToUpdate);
+
+			OnDataUpdatedForStat(FoundBundleStats);
+		}
+
+		void FPersistentStatContainerBase::UpdateSessionPersistentStatTimer(FString SessionName, ETimingStatNames TimerToUpdate)
+		{
+			FSessionPersistentStats& FoundSessionStats = SessionPersistentStatMap.FindOrAdd(SessionName, FSessionPersistentStats(SessionName));
+			FoundSessionStats.UpdateTimingStat(TimerToUpdate);
+
+			OnDataUpdatedForStat(FoundSessionStats);
+		}
+
+		void FPersistentStatContainerBase::IncrementBundlePersistentCounter(FName BundleName, ECountStatNames CounterToUpdate)
+		{
+			FBundlePersistentStats& FoundBundleStats = PerBundlePersistentStatMap.FindOrAdd(BundleName, FBundlePersistentStats(BundleName));
+			FoundBundleStats.IncrementCountStat(CounterToUpdate);
+
+			OnDataUpdatedForStat(FoundBundleStats);
+		}
+
+		void FPersistentStatContainerBase::IncrementSessionPersistentCounter(FString SessionName, ECountStatNames CounterToUpdate)
+		{
+			FSessionPersistentStats& FoundSessionStats = SessionPersistentStatMap.FindOrAdd(SessionName, FSessionPersistentStats(SessionName));
+			FoundSessionStats.IncrementCountStat(CounterToUpdate);
+
+			OnDataUpdatedForStat(FoundSessionStats);
+		}
+
+		void FPersistentStatContainerBase::OnApp_EnteringBackground()
+		{
+			OnBackground_HandleBundleStats();
+			OnBackground_HandleSessionStats();
+		}
+
+		void FPersistentStatContainerBase::OnApp_EnteringForeground()
+		{
+			OnForeground_HandleBundleStats();
+			OnForeground_HandleSessionStats();
+		}
+
+		void FPersistentStatContainerBase::OnBackground_HandleBundleStats()
+		{
+			for (TPair<FName, FBundlePersistentStats>& BundlePair : PerBundlePersistentStatMap)
+			{
+				//Only bother updating bundles listed as active
+				if (BundlePair.Value.IsActive())
+				{
+					UpdateStatsForBackground(BundlePair.Value);
+				}
+			}
+		}
+
+		void FPersistentStatContainerBase::OnForeground_HandleBundleStats()
+		{
+			for (TPair<FName, FBundlePersistentStats>& BundlePair : PerBundlePersistentStatMap)
+			{
+				//Only bother updating bundles listed as active
+				if (BundlePair.Value.IsActive())
+				{
+					UpdateStatsForForeground(BundlePair.Value);
+				}
+			}
+		}
+
+		void FPersistentStatContainerBase::OnBackground_HandleSessionStats()
+		{
+			for (TPair<FString, FSessionPersistentStats>& SessionPair : SessionPersistentStatMap)
+			{
+				//Only bother updating sessions listed as active
+				if (SessionPair.Value.IsActive())
+				{
+					UpdateStatsForBackground(SessionPair.Value);
+				}
+			}
+		}
+
+		void FPersistentStatContainerBase::OnForeground_HandleSessionStats()
+		{
+			for (TPair<FString, FSessionPersistentStats>& SessionPair : SessionPersistentStatMap)
+			{
+				//Only bother updating sessions listed as active
+				if (SessionPair.Value.IsActive())
+				{
+					UpdateStatsForForeground(SessionPair.Value);
+				}
+			}
+		}
+
+		void FPersistentStatContainerBase::UpdateStatsForBackground(FPersistentStatsBase& StatToUpdate)
+		{
+			StatToUpdate.IncrementCountStat(ECountStatNames::NumBackgrounded);
+
+			//Always handle ActiveTotalTime as this isn't dependent on what stage of the process we are in
+			if (StatToUpdate.IsTimingStatStarted(ETimingStatNames::TotalTime_FG))
+			{
+				StatToUpdate.StartTimingStat(ETimingStatNames::TotalTime_BG);
+				StatToUpdate.StopTimingStat(ETimingStatNames::TotalTime_FG);
+			}
+
+			//Besides the ActiveTotalTime above, we should only be in 1 of the following states at a time, so only handle the appropriate swap
+			if (StatToUpdate.IsTimingStatStarted(ETimingStatNames::ChunkDBDownloadTime_FG))
+			{
+				StatToUpdate.StartTimingStat(ETimingStatNames::ChunkDBDownloadTime_BG);
+				StatToUpdate.StopTimingStat(ETimingStatNames::ChunkDBDownloadTime_FG);
+			}
+			else if (StatToUpdate.IsTimingStatStarted(ETimingStatNames::InstallTime_FG))
+			{
+				StatToUpdate.StartTimingStat(ETimingStatNames::InstallTime_BG);
+				StatToUpdate.StopTimingStat(ETimingStatNames::InstallTime_FG);
+			}
+			else if (StatToUpdate.IsTimingStatStarted(ETimingStatNames::PSOTime_FG))
+			{
+				StatToUpdate.StartTimingStat(ETimingStatNames::PSOTime_BG);
+				StatToUpdate.StopTimingStat(ETimingStatNames::PSOTime_FG);
+			}
+
+			OnDataUpdatedForStat(StatToUpdate);
+		}
+
+		void FPersistentStatContainerBase::UpdateStatsForForeground(FPersistentStatsBase& StatToUpdate)
+		{
+			StatToUpdate.IncrementCountStat(ECountStatNames::NumResumedFromBackground);
+
+			//Always handle ActiveTotalTime as this isn't dependent on what stage of the process we are in
+			if (StatToUpdate.IsTimingStatStarted(ETimingStatNames::TotalTime_BG))
+			{
+				StatToUpdate.StopTimingStat(ETimingStatNames::TotalTime_BG);
+				StatToUpdate.StartTimingStat(ETimingStatNames::TotalTime_FG);
+			}
+
+			//Besides the ActiveTotalTime above, we should only be in 1 of the following states at a time, so only handle the appropriate swap
+			if (StatToUpdate.IsTimingStatStarted(ETimingStatNames::ChunkDBDownloadTime_BG))
+			{
+				StatToUpdate.StopTimingStat(ETimingStatNames::ChunkDBDownloadTime_BG);
+				StatToUpdate.StartTimingStat(ETimingStatNames::ChunkDBDownloadTime_FG);
+			}
+			else if (StatToUpdate.IsTimingStatStarted(ETimingStatNames::InstallTime_BG))
+			{
+				StatToUpdate.StopTimingStat(ETimingStatNames::InstallTime_BG);
+				StatToUpdate.StartTimingStat(ETimingStatNames::InstallTime_FG);
+			}
+			else if (StatToUpdate.IsTimingStatStarted(ETimingStatNames::PSOTime_BG))
+			{
+				StatToUpdate.StopTimingStat(ETimingStatNames::PSOTime_BG);
+				StatToUpdate.StartTimingStat(ETimingStatNames::PSOTime_FG);
+			}
+
+			OnDataUpdatedForStat(StatToUpdate);
+		}
+
+		void FPersistentStatContainerBase::UpdateAllBundlesActiveTimers()
+		{
+			TArray<FName> BundleNames;
+			PerBundlePersistentStatMap.GetKeys(BundleNames);
+			
+			for (FName BundleName : BundleNames)
+			{
+				InstallBundleUtil::PersistentStats::FBundlePersistentStats& BundleStats = PerBundlePersistentStatMap.FindOrAdd(BundleName, FBundlePersistentStats(BundleName));
+				BundleStats.UpdateAllActiveTimers();
+
+				OnDataUpdatedForStat(BundleStats);
+			}
+		}
+
+		void FPersistentStatContainerBase::UpdateAllSessionActiveTimers()
+		{
+			TArray<FString> SessionNames;
+			SessionPersistentStatMap.GetKeys(SessionNames);
+
+			for (const FString& SessionName : SessionNames)
+			{
+				InstallBundleUtil::PersistentStats::FSessionPersistentStats& SessionStats = SessionPersistentStatMap.FindOrAdd(SessionName, FSessionPersistentStats(SessionName));
+				SessionStats.UpdateAllActiveTimers();
+
+				OnDataUpdatedForStat(SessionStats);
+			}
+		}
+
+		const FString FPersistentStatsBase::GetBaseExpectedAnalyticsID()
+		{
+			const FString BaseExpectedAnalyticsID = FPlatformMisc::GetDeviceId() + TEXT("_") + FApp::GetBuildVersion();
+			return BaseExpectedAnalyticsID;
+		}
+
+		const FBundlePersistentStats* FPersistentStatContainerBase::GetBundleStat(FName BundleName) const
+		{
+			return PerBundlePersistentStatMap.Find(BundleName);
+		}
+
+		const FSessionPersistentStats* FPersistentStatContainerBase::GetSessionStat(const FString& SessionName) const
+		{
+			return SessionPersistentStatMap.Find(SessionName);
+		}
+	}
 }
diff --git a/Engine/Source/Runtime/InstallBundleManager/Public/InstallBundleManagerInterface.h b/Engine/Source/Runtime/InstallBundleManager/Public/InstallBundleManagerInterface.h
--- a/Engine/Source/Runtime/InstallBundleManager/Public/InstallBundleManagerInterface.h
+++ b/Engine/Source/Runtime/InstallBundleManager/Public/InstallBundleManagerInterface.h
@@ -109,5 +109,8 @@ public:
 	virtual void SetErrorSimulationCommands(const FString& CommandLine) {}
 
 	virtual TSharedPtr<IAnalyticsProviderET> GetAnalyticsProvider() const { return TSharedPtr<IAnalyticsProviderET>(); }
+
+	virtual void StartPersistentStatTrackingSession(const FString& SessionName, const TArray<FName>& RequiredBundles = TArray<FName>(), const FString& ExpectedAnalyticsID = FString(), bool bForceResetStatData = false) {}
+	virtual void StopPersistentStatTrackingSession(const FString& SessionName) {}
 };
 
diff --git a/Engine/Source/Runtime/InstallBundleManager/Public/InstallBundleSourceInterface.h b/Engine/Source/Runtime/InstallBundleManager/Public/InstallBundleSourceInterface.h
--- a/Engine/Source/Runtime/InstallBundleManager/Public/InstallBundleSourceInterface.h
+++ b/Engine/Source/Runtime/InstallBundleManager/Public/InstallBundleSourceInterface.h
@@ -28,7 +28,8 @@ public:
 	virtual void Init(
 		TSharedRef<FConfigFile> InstallBundleConfig, 
 		TSharedRef<InstallBundleUtil::FContentRequestStatsMap> InRequestStats,
-		TSharedPtr<IAnalyticsProviderET> AnalyticsProvider) = 0;
+		TSharedPtr<IAnalyticsProviderET> AnalyticsProvider,
+		TSharedPtr<InstallBundleUtil::PersistentStats::FPersistentStatContainerBase> PersistentStatsContainer) = 0;
 	// Bundle manager will not call AsyncInit again until the bundle source calls back that it is complete
 	// It will be retried indefinitely until init is successful.  
 	virtual void AsyncInit(FInstallBundleSourceInitDelegate Callback) = 0;
diff --git a/Engine/Source/Runtime/InstallBundleManager/Public/InstallBundleTypes.h b/Engine/Source/Runtime/InstallBundleManager/Public/InstallBundleTypes.h
--- a/Engine/Source/Runtime/InstallBundleManager/Public/InstallBundleTypes.h
+++ b/Engine/Source/Runtime/InstallBundleManager/Public/InstallBundleTypes.h
@@ -108,6 +108,7 @@ enum class EInstallBundleRequestFlags : uint32
 	UseBackgroundDownloads = (1 << 1),
 	SendNotificationIfDownloadCompletesInBackground = (1 << 2),
 	ForceNoPatching = (1 << 3),
+	TrackPersistentBundleStats = (1 << 4),
 	Defaults = UseBackgroundDownloads,
 };
 ENUM_CLASS_FLAGS(EInstallBundleRequestFlags)
diff --git a/Engine/Source/Runtime/InstallBundleManager/Public/InstallBundleUtils.h b/Engine/Source/Runtime/InstallBundleManager/Public/InstallBundleUtils.h
--- a/Engine/Source/Runtime/InstallBundleManager/Public/InstallBundleUtils.h
+++ b/Engine/Source/Runtime/InstallBundleManager/Public/InstallBundleUtils.h
@@ -6,6 +6,7 @@
 #include "InstallBundleTypes.h"
 #include "Async/AsyncWork.h"
 #include "Misc/EmbeddedCommunication.h"
+#include "Serialization/JsonSerializerMacros.h"
 
 namespace InstallBundleUtil
 {
@@ -168,6 +169,368 @@ namespace InstallBundleUtil
 		const TMap<FName, InstallBundleUtil::FContentRequestStats>& GetMap() { return StatsMap; }
 	};
 
+	namespace PersistentStats
+	{
+		//Stats used to store timing information in FPersistentTimerData format
+		//_Real stats account for total time since a stat was started regardless of the app's state (shutdown, in foreground, etc.)
+		//_FG stats account for time while the app is active
+		//_BG stats are an estimate of time while the app is inactive (backgrounded, closed, etc.)
+		//		*_BG stats are an estimate as we can not account for unexpected crashes, etc.
+		enum class INSTALLBUNDLEMANAGER_API ETimingStatNames : uint8
+		{
+			  TotalTime_Real
+			, TotalTime_FG
+			, TotalTime_BG
+			, ChunkDBDownloadTime_Real
+			, ChunkDBDownloadTime_FG
+			, ChunkDBDownloadTime_BG
+			, InstallTime_Real
+			, InstallTime_FG
+			, InstallTime_BG
+			, PSOTime_Real
+			, PSOTime_FG
+			, PSOTime_BG
+			, NumStatNames //should always be final entry
+		};
+
+		//Stats used to store count information in an int format
+		enum class INSTALLBUNDLEMANAGER_API ECountStatNames : uint8
+		{
+			  NumResumedFromBackground
+			, NumResumedFromLaunch
+			, NumBackgrounded
+			, NumStatNames //should always be final entry
+		};
+
+		INSTALLBUNDLEMANAGER_API const FString& LexToString(ETimingStatNames InType);
+		INSTALLBUNDLEMANAGER_API const FString& LexToString(ECountStatNames InType);
+		
+		//Returns if the timer is in the _Real set
+		INSTALLBUNDLEMANAGER_API bool IsTimerReal(ETimingStatNames InTimer);
+		//Returns if the timer is in the _FG set
+		INSTALLBUNDLEMANAGER_API bool IsTimerFG(ETimingStatNames InTimer);
+		//Returns if the timer is in the _BG set
+		INSTALLBUNDLEMANAGER_API bool IsTimerBG(ETimingStatNames InTimer);
+
+		//Gets the _Real timer when supplied with the corresponding _FG or _BG timer name
+		INSTALLBUNDLEMANAGER_API ETimingStatNames GetAssociatedRealTimerName(ETimingStatNames InTimerType);
+		//Gets the _FG timer when supplied with the corresponding _Real or _BG timer name
+		INSTALLBUNDLEMANAGER_API ETimingStatNames GetAssociatedFGTimerName(ETimingStatNames InTimerType);
+		//Gets the _BG timer when supplied with the corresponding _Real or _FG timer name
+		INSTALLBUNDLEMANAGER_API ETimingStatNames GetAssociatedBGTimerName(ETimingStatNames InTimerType);
+
+		class FPersistentTimerData : public FJsonSerializable
+		{
+		public:
+			BEGIN_JSON_SERIALIZER
+				JSON_SERIALIZE("LastUpdateTime", LastUpdateTime);
+				JSON_SERIALIZE("TotalTime", CurrentValue);
+			END_JSON_SERIALIZER
+
+		public:
+			FPersistentTimerData()
+				: LastUpdateTime(0.)
+				, CurrentValue(0.)
+			{}
+
+			double LastUpdateTime;
+			double CurrentValue;
+		};
+
+		//Both ContentRequests and Individual Bundles use this to track data respective to themselves.
+		//Keeps us from having to write serialization / loading data for both cases that would be the same
+		class INSTALLBUNDLEMANAGER_API FPersistentStatsBase : public FJsonSerializable
+		{
+		public:
+			//We should not ever really be serializing this base class. Kept it here for reference.
+			BEGIN_JSON_SERIALIZER
+				JSON_SERIALIZE("AnalyticsSessionID", AnalyticsSessionID);
+				JSON_SERIALIZE_MAP("CountStats", CountStatMap);
+				JSON_SERIALIZE_MAP_SERIALIZABLE("TimingStats", TimingStatsMap, FPersistentTimerData);
+			END_JSON_SERIALIZER
+		
+		public:
+			virtual ~FPersistentStatsBase() {}
+
+			//Tries to load persistent stats from disk and determine if we need to resume our previous persistent session or start a new one
+			void StatsBegin(const FString& ExpectedAnalyticsID, bool bForceResetData = false);
+
+			//Tries to set this stat pool as "inactive". By default also stop all active timers. 
+			// NOTE: If you don't stop all active timers future calls to UpdateTimingStat will still increment this value and save it off!
+			void StatsEnd(bool bStopAllActiveTimers = true);
+
+			bool IsTimingStatStarted(ETimingStatNames StatToUpdate) const;
+			void StartTimingStat(ETimingStatNames StatToUpdate);
+			void StopTimingStat(ETimingStatNames StatToUpdate, bool UpdateTimerOnStop = true);
+			void UpdateTimingStat(ETimingStatNames StatToUpdate);
+			
+			void UpdateAllActiveTimers();
+			void StopAllActiveTimers();
+			
+			bool IsActive() const { return bIsActive; }
+			bool IsDirty() const { return bIsDirty; }
+
+			void IncrementCountStat(PersistentStats::ECountStatNames StatToUpdate);
+
+			bool HasTimingStat(ETimingStatNames StatToCheck) const;
+			bool HasCountStat(ECountStatNames StatToCheck) const;
+
+			const FPersistentTimerData* GetTimingStatData(ETimingStatNames StatToGet) const;
+			const int* GetCountStatData(ECountStatNames StatToGet) const;
+
+			//Saves the persistent stat data to disk in the location returned by GetFullPathForStatFile()
+			//Returns True if that save succeeds and false otherwise.
+			bool SaveStatsToDisk();
+
+		protected:
+			bool LoadStatsFromDisk();
+			void ResetStats(const FString& NewAnalyticsSessionID);
+
+			//Called after we load data but before we process any of it.
+			//Useful for fixing up the data loaded from disk before it changes or is acted upon by anything else
+			void OnLoadingDataFromDisk();
+
+			//Helper to try and reconcile offline and active timers after we load data from disk
+			void HandleTimerStatsAfterDataLoad();
+
+			//Protected so we can only construct the derived versions of this Base class
+			FPersistentStatsBase()
+				: TimingStatsMap()
+				, CountStatMap()
+				, AnalyticsSessionID()
+				, bIsActive(false)
+			{
+			}
+			
+		protected:
+			//Serialized Stat Data
+			TMap<FString, FPersistentTimerData> TimingStatsMap;
+			TMap<FString, int> CountStatMap;
+			FString AnalyticsSessionID;
+
+			//We don't serialize these as they are tracking behavior and not stats
+			bool bIsActive;
+			bool bIsDirty;
+			
+		//Static Methods
+		public:
+			//Helper function that generates a basic ExpectedAnalyticsID from device and UE4 information
+			static const FString GetBaseExpectedAnalyticsID();
+		
+		//Pure Virtual Methods
+		public:
+			//Returns Full Path to look for our persistent stat cache
+			virtual const FString GetFullPathForStatFile() const = 0;
+		};
+
+		//Each bundles persistent stats
+		class INSTALLBUNDLEMANAGER_API FBundlePersistentStats : public FPersistentStatsBase
+		{
+		public:
+			BEGIN_JSON_SERIALIZER
+				JSON_SERIALIZE("BundleName", BundleName);
+				JSON_SERIALIZE("AnalyticsSessionID", AnalyticsSessionID);
+				JSON_SERIALIZE_MAP("CountStats", CountStatMap);
+				JSON_SERIALIZE_MAP_SERIALIZABLE("TimingStats", TimingStatsMap, FPersistentTimerData);
+			END_JSON_SERIALIZER
+
+
+		private:
+			//Remove default constructor as we require a BundleName in the constructor
+			FBundlePersistentStats() = delete;
+
+			//Used in our full path and should signify the FString version of the FName of this bundle we are tracking
+			FString BundleName;
+
+		public:
+			FBundlePersistentStats(FName BundleNameIn)
+				: BundleName(BundleNameIn.ToString())
+			{
+			}
+
+			virtual ~FBundlePersistentStats() {}
+
+			//Implementation for FInstallBundleStatsContainerBase
+		public:
+			virtual const FString GetFullPathForStatFile() const override;
+		};
+
+		//Tracks a set of bundles and stats for a named session. This way you can have stats across multiple bundle installs grouped by some naming convention
+		class INSTALLBUNDLEMANAGER_API FSessionPersistentStats : public FPersistentStatsBase
+		{
+		public:
+			BEGIN_JSON_SERIALIZER
+				JSON_SERIALIZE("SessionName", SessionName);
+				JSON_SERIALIZE_ARRAY("RequiredBundles", RequiredBundles);
+				JSON_SERIALIZE("AnalyticsSessionID", AnalyticsSessionID);
+				JSON_SERIALIZE_MAP("CountStats", CountStatMap);
+				JSON_SERIALIZE_MAP_SERIALIZABLE("TimingStats", TimingStatsMap, FPersistentTimerData);
+			END_JSON_SERIALIZER
+
+		private:
+			//Remove default constructor as we require a ContentRequestName in the constructor
+			FSessionPersistentStats() = delete;
+
+			//Used in our full path and should signify some identifier for a given ContentRequest
+			FString SessionName;
+
+			//Used to keep track of all bundles that are part of this content request
+			TArray<FString> RequiredBundles;
+
+		public:
+			FSessionPersistentStats(const FString& ContentRequestNameIn)
+				: SessionName(ContentRequestNameIn)
+				, RequiredBundles()
+			{
+			}
+
+			virtual ~FSessionPersistentStats() {}
+
+			void AddRequiredBundles(const TArray<FString>& RequiredBundlesToAdd);
+			void AddRequiredBundles(const TArray<FName>& RequiredBundlesToAdd);
+
+			//Removes all Required Bundle Names except the ones in the NewRequiredBundlesList. If no NewRequiredBundlesList is supplied, just erases all RequiredBundles.
+			//This includes ones loaded from the persistent cache from previous runs. This re-saves the persistent cache file so they will be removed completely.
+			void ResetRequiredBundles(const TArray<FString>& NewRequiredBundles = TArray<FString>());
+
+			//Populates the passed in TArray with all the required bundles' names
+			void GetRequiredBundles(TArray<FString>& OutRequiredBundles) const;
+			
+			//Implementation for FInstallBundleStatsContainerBase
+		public:
+			virtual const FString GetFullPathForStatFile() const override;
+		};
+
+		//Helper class to handle holding a collection of Bundle and Session stats.
+		class INSTALLBUNDLEMANAGER_API FPersistentStatContainerBase
+		{
+		public:
+			FPersistentStatContainerBase();
+			virtual ~FPersistentStatContainerBase();
+
+			//NOTE! These 2 functions are not virtual as we call them in the construct/deconstructor!
+			void InitializeBase();
+			void ShutdownBase();
+
+			//Starts Bundle Persistent Stat Tracking for the given BundleName. Will automatically load applicable data from the disk and reset that cache
+			//if the ExpectedAnalyticsID doesn't match the one previously serialized.
+			//Uses FInstallBundleStatsContainerBase::GetBaseExpectedAnalyticsID if an ExpectedAnalyticsID is not passed in or empty
+			virtual void StartBundlePersistentStatTracking(FName BundleName, const FString& ExpectedAnalyticsID = FString(), bool bForceResetStatData = false);
+			
+			//Starts Session Persistent Stat Tracking under the given SessionName. Will automatically load applicable data from the disk and reset that cache
+			//if the ExpectedAnalyticsID doesn't match the one previously serialized. 
+			//Uses FInstallBundleStatsContainerBase::GetBaseExpectedAnalyticsID if an ExpectedAnalyticsID is not passed in or empty
+			virtual void StartSessionPersistentStatTracking(const FString& SessionName, const TArray<FName>& RequiredBundles = TArray<FName>(), const FString& ExpectedAnalyticsID = FString(), bool bForceResetStatData = false);
+
+			virtual void StopBundlePersistentStatTracking(FName BundleName, bool bStopAllActiveTimers = true);
+			virtual void StopSessionPersistentStatTracking(FString SessionName, bool bStopAllActiveTimers = true);
+
+			virtual void StartBundlePersistentStatTimer(FName BundleName, ETimingStatNames TimerToStart);
+			virtual void StartSessionPersistentStatTimer(FString SessionName, ETimingStatNames TimerToStart);
+
+			virtual void StopBundlePersistentStatTimer(FName BundleName, ETimingStatNames TimerToStop);
+			virtual void StopSessionPersistentStatTimer(FString SessionName, ETimingStatNames TimerToStop);
+
+			virtual void UpdateBundlePersistentStatTimer(FName BundleName, ETimingStatNames TimerToUpdate);
+			virtual void UpdateSessionPersistentStatTimer(FString SessionName, ETimingStatNames TimerToUpdate);
+
+			virtual void IncrementBundlePersistentCounter(FName BundleName, ECountStatNames CounterToUpdate);
+			virtual void IncrementSessionPersistentCounter(FString SessionName, ECountStatNames CounterToUpdate);
+
+			virtual const FBundlePersistentStats* GetBundleStat(FName BundleName) const;
+			virtual const FSessionPersistentStats* GetSessionStat(const FString& SessionName) const;
+
+			virtual void SaveAllDirtyStatsToDisk();
+
+		protected:
+			virtual bool Tick(float dt);
+			virtual void ResetTimerUpdate();
+			virtual void ResetDirtyStatUpdate();
+
+			virtual void UpdateAllBundlesActiveTimers();
+			virtual void UpdateAllSessionActiveTimers();
+
+			virtual void OnApp_EnteringBackground();
+			virtual void OnApp_EnteringForeground();
+
+			virtual void OnBackground_HandleBundleStats();
+			virtual void OnForeground_HandleBundleStats();
+			virtual void OnBackground_HandleSessionStats();
+			virtual void OnForeground_HandleSessionStats();
+			
+			//Called whenever we start a timer
+			//NOTE: doesn't get called by timer's being Start/Stopped for bShouldAutoHandleFGBGStats
+			virtual void OnTimerStartedForStat(FPersistentStatsBase& BundleStatForTimer, ETimingStatNames TimerStarted);
+
+			//Called whenever we stop a timer 
+			//NOTE: doesn't get called by timer's being Start/Stopped for bShouldAutoHandleFGBGStats
+			virtual void OnTimerStoppedForStat(FPersistentStatsBase& BundleStatForTimer, ETimingStatNames TimerStarted);
+
+			//Called whenever we update our stat's data in one of this class' methods
+			virtual void OnDataUpdatedForStat(FPersistentStatsBase& UpdatedBundleStat);
+
+			//Stops active foreground timers when going to background and starts applicable background version
+			//Also increments background counters
+			virtual void UpdateStatsForBackground(FPersistentStatsBase& StatToUpdate);
+
+			//Stops active background timers and resumes applicable foreground timers
+			//Also increments foreground counters
+			virtual void UpdateStatsForForeground(FPersistentStatsBase& StatToUpdate);
+
+		protected:
+			TMap<FName, FBundlePersistentStats> PerBundlePersistentStatMap;
+			TMap<FString, FSessionPersistentStats> SessionPersistentStatMap;
+
+			FDelegateHandle TickHandle;
+			FDelegateHandle OnApp_EnteringForegroundHandle;
+			FDelegateHandle OnApp_EnteringBackgroundHandle;
+
+			//internal counter used to track when we need to run our timer update during tick. Requires bShouldAutoUpdateTimersInTick to be set to true to effect anything.
+			//reset to TimerAutoUpdateRate whenever we trigger the auto update.
+			float TimerAutoUpdateTimeRemaining;
+
+			//internal counter used to track when we need to update all of our dirty stats. Requires bShouldSaveDirtyStatsOnTick to be set to true to effect anything.
+			//reset to TimerAutoUpdateRate whenever we trigger the dirty stat update
+			float TimerDirtyStatUpdateTimeRemaining;
+			
+
+			/****** Config Values ******/
+			// All values are loaded from GEngineIni in [InstallBundleManager.PersistentStatSettings]
+
+			/****** Timer Auto Update Settings ****/
+			//These settings determine if we auto-update our Timer stats on a tick. Without an auto-update timers would only update on Start / Stop potentially leaving us vulnerable to
+			//dropping a lot of timing information in the event of a crash. By setting these values you ensure your data is accurate to some interval.
+			//
+			//bShouldAutoUpdateTimersInTick -- Determines if the auto update feature is on at all. Rate of auto timer updates is determined by TimerAutoUpdateRate.
+			//
+			//TimerAutoUpdateRate -- Determines the rate we trigger the update in seconds. <=0 means update every tick.
+			//
+			bool bShouldAutoUpdateTimersInTick;
+			float TimerAutoUpdateRate;
+
+			/****** Stat Save To Disk Settings ****/
+			//These settings determine how we handle saving our stats to disk. Outside of these settings, stats ALWAYS try to save to disk on Start or End, but these settings
+			//let us also configure other times we may want to save the data to disk.
+			//
+			//bShouldSaveDirtyStatsOnTick -- If this is set to true then all stats that have bIsDirty set to true will try to save to disk every X seconds. Update rate determined by DirtyStatSaveToDiskRate
+			//								 This way we will not update a given stat more then once a tick to disk.
+			//
+			//DirtyStatSaveToDiskRate -- Determines how often we update our dirty stats in seconds. <=0 means update every tick.
+			//
+			//bShouldSaveStatsOnUpdate -- If this is set to true then we save all stats EVERY time they are updated without waiting for tick.
+			//								 This could lead to multiple disk writes for a given stat in 1 tick.
+			bool bShouldSaveDirtyStatsOnTick;
+			float DirtyStatSaveToDiskRate;
+			bool bShouldSaveStatsOnUpdate;
+
+			//Determines if we should automatically Start,Stop,and Update the _FG and _BG versions of our timer stats.
+			//If true you only have to call Start/Stop on the _Real versions of all timers and we will automatically
+			//Start and Stop the _FG and _BG timers when going to background or returning to foreground to keep
+			//their data accurate.
+			bool bShouldAutoHandleFGBGStats;
+		};
+	}
+
 	struct INSTALLBUNDLEMANAGER_API IBundleSourceContentRequestSharedContext
 	{
 		virtual ~IBundleSourceContentRequestSharedContext() {}
