Add LexFromString for EInstallBundleSourceType.

diff --git a/Engine/Source/Runtime/InstallBundleManager/Private/InstallBundleTypes.cpp b/Engine/Source/Runtime/InstallBundleManager/Private/InstallBundleTypes.cpp
--- a/Engine/Source/Runtime/InstallBundleManager/Private/InstallBundleTypes.cpp
+++ b/Engine/Source/Runtime/InstallBundleManager/Private/InstallBundleTypes.cpp
@@ -4,6 +4,7 @@
 #include "InstallBundleManagerPrivatePCH.h"
 
 #include "InstallBundleUtils.h"
+#include "Misc/CString.h"
 
 const TCHAR* LexToString(EInstallBundleSourceType Type)
 {
@@ -18,6 +19,21 @@ const TCHAR* LexToString(EInstallBundleSourceType Type)
 	return Strings[InstallBundleUtil::CastToUnderlying(Type)];
 }
 
+void LexFromString(EInstallBundleSourceType& OutType, const TCHAR* String)
+{
+	OutType = EInstallBundleSourceType::Count;
+
+	for (EInstallBundleSourceType SourceType : TEnumRange<EInstallBundleSourceType>())
+	{
+		const TCHAR* SourceStr = LexToString(SourceType);
+		if (FCString::Stricmp(SourceStr, String) == 0)
+		{
+			OutType = SourceType;
+			break;
+		}
+	}
+}
+
 const TCHAR* LexToString(EInstallBundleManagerInitResult Result)
 {
 	static const TCHAR* Strings[] =
diff --git a/Engine/Source/Runtime/InstallBundleManager/Public/InstallBundleTypes.h b/Engine/Source/Runtime/InstallBundleManager/Public/InstallBundleTypes.h
--- a/Engine/Source/Runtime/InstallBundleManager/Public/InstallBundleTypes.h
+++ b/Engine/Source/Runtime/InstallBundleManager/Public/InstallBundleTypes.h
@@ -14,6 +14,7 @@ enum class EInstallBundleSourceType : int
 };
 ENUM_RANGE_BY_COUNT(EInstallBundleSourceType, EInstallBundleSourceType::Count);
 INSTALLBUNDLEMANAGER_API const TCHAR* LexToString(EInstallBundleSourceType Type);
+INSTALLBUNDLEMANAGER_API void LexFromString(EInstallBundleSourceType& OutType, const TCHAR* String);
 
 enum class EInstallBundleManagerInitState : int
 {
