Install Bundle Source Interface
- Add GetBundleSkipReason() to allow bundle sources to reject bundles early.  This prevents attempting to wait on language chunks that don't match the current language mask.  Currently not tested because I need a new build.

diff --git a/Engine/Source/Runtime/InstallBundleManager/Public/InstallBundleSourceInterface.h b/Engine/Source/Runtime/InstallBundleManager/Public/InstallBundleSourceInterface.h
--- a/Engine/Source/Runtime/InstallBundleManager/Public/InstallBundleSourceInterface.h
+++ b/Engine/Source/Runtime/InstallBundleManager/Public/InstallBundleSourceInterface.h
@@ -43,6 +43,11 @@ public:
 	// BundleNames contains all dependencies and has been deduped
 	virtual void GetContentState(TArrayView<const FName> BundleNames, EInstallBundleGetContentStateFlags Flags, FInstallBundleGetContentStateDelegate Callback) = 0;
 
+	// Allows this bundle source to reject bundle requests early, rather than failing them.
+	// This means that client code does not have to wait on these bundles.
+	// This may be called anytime after Init, even during AsyncInit
+	virtual EInstallBundleSourceBundleSkipReason GetBundleSkipReason(FName BundleName) const { return EInstallBundleSourceBundleSkipReason::None; }
+
 	struct FRequestUpdateContentBundleContext
 	{
 		FName BundleName;
diff --git a/Engine/Source/Runtime/InstallBundleManager/Public/InstallBundleTypes.h b/Engine/Source/Runtime/InstallBundleManager/Public/InstallBundleTypes.h
--- a/Engine/Source/Runtime/InstallBundleManager/Public/InstallBundleTypes.h
+++ b/Engine/Source/Runtime/InstallBundleManager/Public/InstallBundleTypes.h
@@ -85,7 +85,9 @@ enum class EInstallBundleRequestInfoFlags : int32
 	EnqueuedBundlesForInstall = (1 << 0),
 	SkippedAlreadyMountedBundles = (1 << 1),
 	SkippedUnknownBundles = (1 << 2),
-	InitializationError = (1 << 3), // Can't enqueue because the bundle manager failed to initialize
+	SkippedUnusableLanguageBundles = (1 << 3), // Can't enqueue language bundles because of current system settings
+	SkippedBundlesDueToBundleSource = (1 << 4), // A bundle source rejected a bundle for some reason
+	InitializationError = (1 << 5), // Can't enqueue because the bundle manager failed to initialize
 };
 ENUM_CLASS_FLAGS(EInstallBundleRequestInfoFlags);
 
@@ -120,6 +122,32 @@ struct FInstallBundleRequestInfo
 	TArray<FName> BundlesQueuedForInstall;
 };
 
+enum class EInstallBundleCancelFlags : uint32
+{
+	None = 0,
+	Resumable = (1 << 0),
+};
+ENUM_CLASS_FLAGS(EInstallBundleCancelFlags);
+
+enum class EInstallBundlePauseFlags : uint32
+{
+	None = 0,
+	OnCellularNetwork = (1 << 0),
+	NoInternetConnection = (1 << 1),
+	UserPaused = (1 << 2)
+};
+ENUM_CLASS_FLAGS(EInstallBundlePauseFlags);
+
+enum class EInstallBundleStatus : int
+{
+	Requested,
+	Updating,
+	Finishing,
+	Ready,
+	Count,
+};
+INSTALLBUNDLEMANAGER_API const TCHAR* LexToString(EInstallBundleStatus Status);
+
 struct FInstallBundleSourceInitInfo
 {
 	EInstallBundleManagerInitResult Result = EInstallBundleManagerInitResult::OK;
@@ -146,21 +174,12 @@ struct FInstallBundleSourceRequestResultInfo
 	bool bContentWasInstalled = false;
 };
 
-enum class EInstallBundleCancelFlags : int32
+struct FInstallBundleSourceProgress
 {
-	None = 0,
-	Resumable = (1 << 0),
-};
-ENUM_CLASS_FLAGS(EInstallBundleCancelFlags);
+	FName BundleName;
 
-enum class EInstallBundlePauseFlags : uint32
-{
-	None = 0,
-	OnCellularNetwork = (1 << 0),
-	NoInternetConnection = (1 << 1),
-	UserPaused = (1 << 2)
+	float Install_Percent = 0;
 };
-ENUM_CLASS_FLAGS(EInstallBundlePauseFlags);
 
 struct FInstallBundleSourcePauseInfo
 {
@@ -171,19 +190,9 @@ struct FInstallBundleSourcePauseInfo
 	bool bDidPauseChange = false;
 };
 
-enum class EInstallBundleStatus : int
+enum class EInstallBundleSourceBundleSkipReason : uint32
 {
-	Requested,
-	Updating,
-	Finishing,
-	Ready,
-	Count,
-};
-INSTALLBUNDLEMANAGER_API const TCHAR* LexToString(EInstallBundleStatus Status);
-
-struct FInstallBundleSourceProgress
-{
-	FName BundleName;
-
-	float Install_Percent = 0;
+	None = 0,
+	LanguageNotCurrent = (1 << 0), // The platform language must be changed to make it valid to request this bundle
 };
+ENUM_CLASS_FLAGS(EInstallBundleSourceBundleSkipReason);
