Install Bundle Manager
- Allow "invalid" (empty, not generated for this build, etc.) bundles to be skipped.

diff --git a/Engine/Source/Runtime/InstallBundleManager/Public/InstallBundleTypes.h b/Engine/Source/Runtime/InstallBundleManager/Public/InstallBundleTypes.h
--- a/Engine/Source/Runtime/InstallBundleManager/Public/InstallBundleTypes.h
+++ b/Engine/Source/Runtime/InstallBundleManager/Public/InstallBundleTypes.h
@@ -85,9 +85,10 @@ enum class EInstallBundleRequestInfoFlags : int32
 	EnqueuedBundlesForInstall = (1 << 0),
 	SkippedAlreadyMountedBundles = (1 << 1),
 	SkippedUnknownBundles = (1 << 2),
-	SkippedUnusableLanguageBundles = (1 << 3), // Can't enqueue language bundles because of current system settings
-	SkippedBundlesDueToBundleSource = (1 << 4), // A bundle source rejected a bundle for some reason
-	InitializationError = (1 << 5), // Can't enqueue because the bundle manager failed to initialize
+	SkippedInvalidBundles = (1 << 3), // Bundle can't be used with this build
+	SkippedUnusableLanguageBundles = (1 << 4), // Can't enqueue language bundles because of current system settings
+	SkippedBundlesDueToBundleSource = (1 << 5), // A bundle source rejected a bundle for some reason
+	InitializationError = (1 << 6), // Can't enqueue because the bundle manager failed to initialize
 };
 ENUM_CLASS_FLAGS(EInstallBundleRequestInfoFlags);
 
@@ -194,5 +195,6 @@ enum class EInstallBundleSourceBundleSkipReason : uint32
 {
 	None = 0,
 	LanguageNotCurrent = (1 << 0), // The platform language must be changed to make it valid to request this bundle
+	NotValid = (1 << 1), // Bundle can't be used with this build
 };
 ENUM_CLASS_FLAGS(EInstallBundleSourceBundleSkipReason);
