Improvements for PersistentStats.
- Added ensures to flag problems where you attempt to start stats on bundles that haven't yet been started.
- Adding ability to pass in an FInstallBundleCombinedContentState to StartPersistentStatTrackingSession on InstallBundleManager as this is a useful thing to react to when starting your session tracking.
- Removing bShouldSaveStatsOnUpdate as this is really a trap that noone should ever do. Instead should just set it to save to disk every tick by setting bShouldSaveDirtyStatsOnTick and DirtyStatSaveToDiskRate to 0.
- Changed LoadFromDisk code so that you can load stats from disk without calling StatsBegin. There was a need for stats to be loaded from disk, but not yet active for Sessions (which now call this code automatically when starting a Session PersistentStats for all of it's RequiredBundles).
#tests iOS Client

diff --git a/Engine/Source/Runtime/InstallBundleManager/Private/InstallBundleUtils.cpp b/Engine/Source/Runtime/InstallBundleManager/Private/InstallBundleUtils.cpp
--- a/Engine/Source/Runtime/InstallBundleManager/Private/InstallBundleUtils.cpp
+++ b/Engine/Source/Runtime/InstallBundleManager/Private/InstallBundleUtils.cpp
@@ -454,17 +454,27 @@ namespace InstallBundleUtil
 
 		bool FPersistentStatsBase::LoadStatsFromDisk()
 		{
-			FString JSONStringOnDisk;
-			if (FPaths::FileExists(GetFullPathForStatFile()))
+			if (!bHasLoadedFromDisk)
 			{
-				FFileHelper::LoadFileToString(JSONStringOnDisk, *GetFullPathForStatFile());
-			}
+				FString JSONStringOnDisk;
+				if (FPaths::FileExists(GetFullPathForStatFile()))
+				{
+					FFileHelper::LoadFileToString(JSONStringOnDisk, *GetFullPathForStatFile());
+				}
 
-			if (!JSONStringOnDisk.IsEmpty())
-			{
-				return FromJson(JSONStringOnDisk);
+				if (!JSONStringOnDisk.IsEmpty())
+				{
+					bHasLoadedFromDisk = FromJson(JSONStringOnDisk);
+					
+					if (bHasLoadedFromDisk)
+					{
+						OnLoadingDataFromDisk();
+					}
+					
+					return bHasLoadedFromDisk;
+				}
 			}
-
+			
 			return false;
 		}
 
@@ -479,7 +489,7 @@ namespace InstallBundleUtil
 			TimingStatsMap.Reset();
 			CountStatMap.Reset();
 			AnalyticsSessionID = NewAnalyticsSessionID;
-
+			
 			bIsDirty = true;
 		}
 
@@ -606,10 +616,7 @@ namespace InstallBundleUtil
 		{
 			bIsActive = true;
 			
-			if (LoadStatsFromDisk())
-			{
-				OnLoadingDataFromDisk();
-			}
+			LoadStatsFromDisk();
 
 			//If our Analytics ID doesn't match our expected we need to reset the data as we have started a new persistent session
 			if (bForceResetData || !AnalyticsSessionID.Equals(ExpectedAnalyticsID))
@@ -727,6 +734,7 @@ namespace InstallBundleUtil
 			{
 				RequiredBundles.AddUnique(BundleName);
 			}
+			
 			bIsDirty = true;
 		}
 
@@ -762,7 +770,7 @@ namespace InstallBundleUtil
 
 		const FString FSessionPersistentStats::GetFullPathForStatFile() const
 		{
-			return FPaths::Combine(FPlatformMisc::GamePersistentDownloadDir(), TEXT("PersistentStats"), TEXT("ContentRequestStats"), (SessionName + TEXT(".json")));
+			return FPaths::Combine(FPlatformMisc::GamePersistentDownloadDir(), TEXT("PersistentStats"), TEXT("SessionStats"), (SessionName + TEXT(".json")));
 		}
 
 		FPersistentStatContainerBase::FPersistentStatContainerBase()
@@ -777,7 +785,6 @@ namespace InstallBundleUtil
 			, TimerAutoUpdateRate(10.0f)
 			, bShouldSaveDirtyStatsOnTick(true)
 			, DirtyStatSaveToDiskRate(5.f)
-			, bShouldSaveStatsOnUpdate(false)
 			, bShouldAutoHandleFGBGStats(true)
 		{
 			InitializeBase();
@@ -797,7 +804,6 @@ namespace InstallBundleUtil
 
 				GConfig->GetBool(TEXT("InstallBundleManager.PersistentStatSettings"), TEXT("bShouldSaveDirtyStatsOnTick"), bShouldSaveDirtyStatsOnTick, GEngineIni);
 				GConfig->GetFloat(TEXT("InstallBundleManager.PersistentStatSettings"), TEXT("DirtyStatSaveToDiskRate"), DirtyStatSaveToDiskRate, GEngineIni);
-				GConfig->GetBool(TEXT("InstallBundleManager.PersistentStatSettings"), TEXT("bShouldSaveStatsOnUpdate"), bShouldSaveStatsOnUpdate, GEngineIni);
 				
 				GConfig->GetBool(TEXT("InstallBundleManager.PersistentStatSettings"), TEXT("bShouldAutoHandleFGBGStats"), bShouldAutoHandleFGBGStats, GEngineIni);
 
@@ -903,14 +909,6 @@ namespace InstallBundleUtil
 			}
 		}
 
-		void FPersistentStatContainerBase::OnDataUpdatedForStat(FPersistentStatsBase& UpdatedBundleStat)
-		{
-			if (bShouldSaveStatsOnUpdate)
-			{
-				UpdatedBundleStat.SaveStatsToDisk();
-			}
-		}
-
 		void FPersistentStatContainerBase::ResetTimerUpdate()
 		{
 			TimerAutoUpdateTimeRemaining = TimerAutoUpdateRate;
@@ -955,8 +953,6 @@ namespace InstallBundleUtil
 
 			FBundlePersistentStats& FoundBundleStats = PerBundlePersistentStatMap.FindOrAdd(BundleName, FBundlePersistentStats(BundleName));
 			FoundBundleStats.StatsBegin(ExpectedAnalyticsToUse, bForceResetStatData);
-
-			OnDataUpdatedForStat(FoundBundleStats);
 		}
 
 		void FPersistentStatContainerBase::StartSessionPersistentStatTracking(const FString& SessionName, const TArray<FName>& RequiredBundles /* = TArray<FName>() */, const FString& ExpectedAnalyticsID /* = FString() */, bool bForceResetStatData /* = false */)
@@ -969,8 +965,9 @@ namespace InstallBundleUtil
 
 			//Also append starting required bundles as we may have new ones from the ones already in data
 			FoundSessionStats.AddRequiredBundles(RequiredBundles);
-
-			OnDataUpdatedForStat(FoundSessionStats);
+			
+			//Go ahead and load data for all bundles in our RequiredBundles list while we are starting our Session
+			LoadRequiredBundleDataFromDiskForSession(SessionName);
 		}
 		
 		void FPersistentStatContainerBase::StopBundlePersistentStatTracking(FName BundleName, bool bStopAllActiveTimers /* = true */)
@@ -979,7 +976,6 @@ namespace InstallBundleUtil
 			if (nullptr != FoundBundleStats)
 			{
 				FoundBundleStats->StatsEnd(bStopAllActiveTimers);
-				OnDataUpdatedForStat(*FoundBundleStats);
 			}
 		}
 
@@ -989,17 +985,42 @@ namespace InstallBundleUtil
 			if (nullptr != FoundSessionStats)
 			{
 				FoundSessionStats->StatsEnd(bStopAllActiveTimers);
-				OnDataUpdatedForStat(*FoundSessionStats);
 			}
 		}
 
+		void FPersistentStatContainerBase::LoadRequiredBundleDataFromDiskForSession(const FString& SessionName)
+		{
+			FSessionPersistentStats* FoundSessionStats = SessionPersistentStatMap.Find(SessionName);
+			if (nullptr != FoundSessionStats)
+			{
+				TArray<FString> RequiredBundles;
+				FoundSessionStats->GetRequiredBundles(RequiredBundles);
+				
+				for (const FString& Bundle : RequiredBundles)
+				{
+					FBundlePersistentStats* FoundBundleStats = PerBundlePersistentStatMap.Find(*Bundle);
+					if (nullptr == FoundBundleStats)
+					{
+						FBundlePersistentStats NewBundleStats = FBundlePersistentStats(*Bundle);
+						NewBundleStats.LoadStatsFromDisk();
+						PerBundlePersistentStatMap.Emplace(*Bundle, MoveTemp(NewBundleStats));
+					}
+					else
+					{
+						FoundBundleStats->LoadStatsFromDisk();
+					}
+				}
+			}
+		}
+	
 		void FPersistentStatContainerBase::StartBundlePersistentStatTimer(FName BundleName, ETimingStatNames TimerToStart)
 		{
 			FBundlePersistentStats& FoundBundleStats = PerBundlePersistentStatMap.FindOrAdd(BundleName, FBundlePersistentStats(BundleName));
 			FoundBundleStats.StartTimingStat(TimerToStart);
 
+			ensureAlwaysMsgf(FoundBundleStats.IsActive(), TEXT("Invalid attempt to start %s on bundle %s that hasn't yet had StartBundlePersistentStatTracking called on it! Should always start tracking before using persistent stats!"), *LexToString(TimerToStart), *(BundleName.ToString()));
+			
 			OnTimerStartedForStat(FoundBundleStats, TimerToStart);
-			OnDataUpdatedForStat(FoundBundleStats);
 		}
 
 		void FPersistentStatContainerBase::StartSessionPersistentStatTimer(FString SessionName, ETimingStatNames TimerToStart)
@@ -1007,8 +1028,9 @@ namespace InstallBundleUtil
 			FSessionPersistentStats& FoundSessionStats = SessionPersistentStatMap.FindOrAdd(SessionName, FSessionPersistentStats(SessionName));
 			FoundSessionStats.StartTimingStat(TimerToStart);
 
+			ensureAlwaysMsgf(FoundSessionStats.IsActive(), TEXT("Invalid attempt to start %s on session %s that hasn't yet had StartBundlePersistentStatTracking called on it! Should always start tracking before using persistent stats!"), *LexToString(TimerToStart), *SessionName);
+			
 			OnTimerStartedForStat(FoundSessionStats, TimerToStart);
-			OnDataUpdatedForStat(FoundSessionStats);
 		}
 
 		void FPersistentStatContainerBase::StopBundlePersistentStatTimer(FName BundleName, ETimingStatNames TimerToStop)
@@ -1017,7 +1039,6 @@ namespace InstallBundleUtil
 			FoundBundleStats.StopTimingStat(TimerToStop);
 
 			OnTimerStoppedForStat(FoundBundleStats, TimerToStop);
-			OnDataUpdatedForStat(FoundBundleStats);
 		}
 
 		void FPersistentStatContainerBase::StopSessionPersistentStatTimer(FString SessionName, ETimingStatNames TimerToStop)
@@ -1026,23 +1047,18 @@ namespace InstallBundleUtil
 			FoundSessionStats.StopTimingStat(TimerToStop);
 
 			OnTimerStoppedForStat(FoundSessionStats, TimerToStop);
-			OnDataUpdatedForStat(FoundSessionStats);
 		}
 
 		void FPersistentStatContainerBase::UpdateBundlePersistentStatTimer(FName BundleName, ETimingStatNames TimerToUpdate)
 		{
 			FBundlePersistentStats& FoundBundleStats = PerBundlePersistentStatMap.FindOrAdd(BundleName, FBundlePersistentStats(BundleName));
 			FoundBundleStats.UpdateTimingStat(TimerToUpdate);
-
-			OnDataUpdatedForStat(FoundBundleStats);
 		}
 
 		void FPersistentStatContainerBase::UpdateSessionPersistentStatTimer(FString SessionName, ETimingStatNames TimerToUpdate)
 		{
 			FSessionPersistentStats& FoundSessionStats = SessionPersistentStatMap.FindOrAdd(SessionName, FSessionPersistentStats(SessionName));
 			FoundSessionStats.UpdateTimingStat(TimerToUpdate);
-
-			OnDataUpdatedForStat(FoundSessionStats);
 		}
 
 		void FPersistentStatContainerBase::IncrementBundlePersistentCounter(FName BundleName, ECountStatNames CounterToUpdate)
@@ -1050,7 +1066,7 @@ namespace InstallBundleUtil
 			FBundlePersistentStats& FoundBundleStats = PerBundlePersistentStatMap.FindOrAdd(BundleName, FBundlePersistentStats(BundleName));
 			FoundBundleStats.IncrementCountStat(CounterToUpdate);
 
-			OnDataUpdatedForStat(FoundBundleStats);
+			ensureAlwaysMsgf(FoundBundleStats.IsActive(), TEXT("Invalid attempt to increment %s on bundle %s that hasn't yet had StartBundlePersistentStatTracking called on it! Should always start tracking before using persistent stats!"), *LexToString(CounterToUpdate), *(BundleName.ToString()));
 		}
 
 		void FPersistentStatContainerBase::IncrementSessionPersistentCounter(FString SessionName, ECountStatNames CounterToUpdate)
@@ -1058,7 +1074,7 @@ namespace InstallBundleUtil
 			FSessionPersistentStats& FoundSessionStats = SessionPersistentStatMap.FindOrAdd(SessionName, FSessionPersistentStats(SessionName));
 			FoundSessionStats.IncrementCountStat(CounterToUpdate);
 
-			OnDataUpdatedForStat(FoundSessionStats);
+			ensureAlwaysMsgf(FoundSessionStats.IsActive(), TEXT("Invalid attempt to increment %s on session %s that hasn't yet had StartBundlePersistentStatTracking called on it! Should always start tracking before using persistent stats!"), *LexToString(CounterToUpdate), *SessionName);
 		}
 
 		void FPersistentStatContainerBase::OnApp_EnteringBackground()
@@ -1148,8 +1164,6 @@ namespace InstallBundleUtil
 				StatToUpdate.StartTimingStat(ETimingStatNames::PSOTime_BG);
 				StatToUpdate.StopTimingStat(ETimingStatNames::PSOTime_FG);
 			}
-
-			OnDataUpdatedForStat(StatToUpdate);
 		}
 
 		void FPersistentStatContainerBase::UpdateStatsForForeground(FPersistentStatsBase& StatToUpdate)
@@ -1179,8 +1193,6 @@ namespace InstallBundleUtil
 				StatToUpdate.StopTimingStat(ETimingStatNames::PSOTime_BG);
 				StatToUpdate.StartTimingStat(ETimingStatNames::PSOTime_FG);
 			}
-
-			OnDataUpdatedForStat(StatToUpdate);
 		}
 
 		void FPersistentStatContainerBase::UpdateAllBundlesActiveTimers()
@@ -1192,8 +1204,6 @@ namespace InstallBundleUtil
 			{
 				InstallBundleUtil::PersistentStats::FBundlePersistentStats& BundleStats = PerBundlePersistentStatMap.FindOrAdd(BundleName, FBundlePersistentStats(BundleName));
 				BundleStats.UpdateAllActiveTimers();
-
-				OnDataUpdatedForStat(BundleStats);
 			}
 		}
 
@@ -1206,8 +1216,6 @@ namespace InstallBundleUtil
 			{
 				InstallBundleUtil::PersistentStats::FSessionPersistentStats& SessionStats = SessionPersistentStatMap.FindOrAdd(SessionName, FSessionPersistentStats(SessionName));
 				SessionStats.UpdateAllActiveTimers();
-
-				OnDataUpdatedForStat(SessionStats);
 			}
 		}
 
diff --git a/Engine/Source/Runtime/InstallBundleManager/Public/InstallBundleManagerInterface.h b/Engine/Source/Runtime/InstallBundleManager/Public/InstallBundleManagerInterface.h
--- a/Engine/Source/Runtime/InstallBundleManager/Public/InstallBundleManagerInterface.h
+++ b/Engine/Source/Runtime/InstallBundleManager/Public/InstallBundleManagerInterface.h
@@ -110,7 +110,7 @@ public:
 
 	virtual TSharedPtr<IAnalyticsProviderET> GetAnalyticsProvider() const { return TSharedPtr<IAnalyticsProviderET>(); }
 
-	virtual void StartPersistentStatTrackingSession(const FString& SessionName, const TArray<FName>& RequiredBundles = TArray<FName>(), const FString& ExpectedAnalyticsID = FString(), bool bForceResetStatData = false) {}
-	virtual void StopPersistentStatTrackingSession(const FString& SessionName) {}
+	virtual void StartSessionPersistentStatTracking(const FString& SessionName, const TArray<FName>& RequiredBundles = TArray<FName>(), const FString& ExpectedAnalyticsID = FString(), bool bForceResetStatData = false, const FInstallBundleCombinedContentState* State = nullptr) {}
+	virtual void StopSessionPersistentStatTracking(const FString& SessionName) {}
 };
 
diff --git a/Engine/Source/Runtime/InstallBundleManager/Public/InstallBundleTypes.h b/Engine/Source/Runtime/InstallBundleManager/Public/InstallBundleTypes.h
--- a/Engine/Source/Runtime/InstallBundleManager/Public/InstallBundleTypes.h
+++ b/Engine/Source/Runtime/InstallBundleManager/Public/InstallBundleTypes.h
@@ -173,6 +173,8 @@ struct FInstallBundleSourceRequestResultInfo
 	TSet<FString> NonUFSShaderLibPaths;
 
 	bool bContentWasInstalled = false;
+	
+	bool DidBundleSourceDoWork() const { return (ContentPaths.Num() != 0);} 
 };
 
 struct FInstallBundleSourceProgress
diff --git a/Engine/Source/Runtime/InstallBundleManager/Public/InstallBundleUtils.h b/Engine/Source/Runtime/InstallBundleManager/Public/InstallBundleUtils.h
--- a/Engine/Source/Runtime/InstallBundleManager/Public/InstallBundleUtils.h
+++ b/Engine/Source/Runtime/InstallBundleManager/Public/InstallBundleUtils.h
@@ -283,13 +283,18 @@ namespace InstallBundleUtil
 			//Saves the persistent stat data to disk in the location returned by GetFullPathForStatFile()
 			//Returns True if that save succeeds and false otherwise.
 			bool SaveStatsToDisk();
+			
+			//Function that allows you to load stats from disk. Returns true if stats were either loaded from disk now, or previously were loaded from disk.
+			//NOTE: This does not begin stat collection, so you should still call StatsBegin on this before handling any stats!
+			bool LoadStatsFromDisk();
 
 		protected:
-			bool LoadStatsFromDisk();
 			void ResetStats(const FString& NewAnalyticsSessionID);
-
+			
 			//Called after we load data but before we process any of it.
 			//Useful for fixing up the data loaded from disk before it changes or is acted upon by anything else
+			//NOTE: A call to LoadStatsFromDisk will not necessarily call this! It will only be called when data is actually pulled from
+			//the disk, and its possible to skip pulling data from the disk!
 			void OnLoadingDataFromDisk();
 
 			//Helper to try and reconcile offline and active timers after we load data from disk
@@ -301,6 +306,8 @@ namespace InstallBundleUtil
 				, CountStatMap()
 				, AnalyticsSessionID()
 				, bIsActive(false)
+				, bIsDirty(false)
+				, bHasLoadedFromDisk(false)
 			{
 			}
 			
@@ -314,6 +321,8 @@ namespace InstallBundleUtil
 			bool bIsActive;
 			bool bIsDirty;
 			
+			bool bHasLoadedFromDisk;
+			
 		//Static Methods
 		public:
 			//Helper function that generates a basic ExpectedAnalyticsID from device and UE4 information
@@ -443,7 +452,7 @@ namespace InstallBundleUtil
 			virtual const FSessionPersistentStats* GetSessionStat(const FString& SessionName) const;
 
 			virtual void SaveAllDirtyStatsToDisk();
-
+			
 		protected:
 			virtual bool Tick(float dt);
 			virtual void ResetTimerUpdate();
@@ -468,9 +477,6 @@ namespace InstallBundleUtil
 			//NOTE: doesn't get called by timer's being Start/Stopped for bShouldAutoHandleFGBGStats
 			virtual void OnTimerStoppedForStat(FPersistentStatsBase& BundleStatForTimer, ETimingStatNames TimerStarted);
 
-			//Called whenever we update our stat's data in one of this class' methods
-			virtual void OnDataUpdatedForStat(FPersistentStatsBase& UpdatedBundleStat);
-
 			//Stops active foreground timers when going to background and starts applicable background version
 			//Also increments background counters
 			virtual void UpdateStatsForBackground(FPersistentStatsBase& StatToUpdate);
@@ -478,7 +484,13 @@ namespace InstallBundleUtil
 			//Stops active background timers and resumes applicable foreground timers
 			//Also increments foreground counters
 			virtual void UpdateStatsForForeground(FPersistentStatsBase& StatToUpdate);
-
+			
+			//Goes through this Session's RequiredBundles and loads data from disk for those bundles. Allows us to make sure we have
+			//full data in memory from disk without having to know the previous runs AnalyticsID to call StatsBegin on that bundle.
+			//NOTE: You should still call StatsBegin on the listed bundles before Starting/Incrementing/ETC any analytics! This just
+			//ensures that the previous run data exists for this given session!
+			virtual void LoadRequiredBundleDataFromDiskForSession(const FString& SessionName);
+			
 		protected:
 			TMap<FName, FBundlePersistentStats> PerBundlePersistentStatMap;
 			TMap<FString, FSessionPersistentStats> SessionPersistentStatMap;
@@ -518,12 +530,8 @@ namespace InstallBundleUtil
 			//								 This way we will not update a given stat more then once a tick to disk.
 			//
 			//DirtyStatSaveToDiskRate -- Determines how often we update our dirty stats in seconds. <=0 means update every tick.
-			//
-			//bShouldSaveStatsOnUpdate -- If this is set to true then we save all stats EVERY time they are updated without waiting for tick.
-			//								 This could lead to multiple disk writes for a given stat in 1 tick.
 			bool bShouldSaveDirtyStatsOnTick;
 			float DirtyStatSaveToDiskRate;
-			bool bShouldSaveStatsOnUpdate;
 
 			//Determines if we should automatically Start,Stop,and Update the _FG and _BG versions of our timer stats.
 			//If true you only have to call Start/Stop on the _Real versions of all timers and we will automatically
