Install Bundle Manager Interface
- Allow Synchronous Init to Fail
- Allow requests stats to be reset

diff --git a/Engine/Source/Runtime/InstallBundleManager/Private/InstallBundleUtils.cpp b/Engine/Source/Runtime/InstallBundleManager/Private/InstallBundleUtils.cpp
--- a/Engine/Source/Runtime/InstallBundleManager/Private/InstallBundleUtils.cpp
+++ b/Engine/Source/Runtime/InstallBundleManager/Private/InstallBundleUtils.cpp
@@ -140,6 +140,11 @@ namespace InstallBundleUtil
 		}
 	}
 
+	void FContentRequestStatsMap::StatsReset(FName BundleName)
+	{
+		StatsMap.Remove(BundleName);
+	}
+
 	void FContentRequestStatsMap::StatsBegin(FName BundleName, const TCHAR* State)
 	{
 		FContentRequestStats& Stats = StatsMap.FindOrAdd(BundleName);
diff --git a/Engine/Source/Runtime/InstallBundleManager/Public/InstallBundleSourceInterface.h b/Engine/Source/Runtime/InstallBundleManager/Public/InstallBundleSourceInterface.h
--- a/Engine/Source/Runtime/InstallBundleManager/Public/InstallBundleSourceInterface.h
+++ b/Engine/Source/Runtime/InstallBundleManager/Public/InstallBundleSourceInterface.h
@@ -8,7 +8,7 @@
 class IInstallBundleSource;
 class IAnalyticsProviderET;
 
-DECLARE_DELEGATE_TwoParams(FInstallBundleSourceInitDelegate, TSharedRef<IInstallBundleSource> /*Source*/, FInstallBundleSourceInitInfo /*InitInfo*/);
+DECLARE_DELEGATE_TwoParams(FInstallBundleSourceInitDelegate, TSharedRef<IInstallBundleSource> /*Source*/, FInstallBundleSourceAsyncInitInfo /*InitInfo*/);
 
 DECLARE_DELEGATE_OneParam(FInstallBundleCompleteDelegate, FInstallBundleSourceRequestResultInfo /*Result*/);
 DECLARE_DELEGATE_OneParam(FInstallBundlePausedDelegate, FInstallBundleSourcePauseInfo /*PauseInfo*/);
@@ -25,7 +25,8 @@ public:
 	virtual float GetSourceWeight() const { return 1.0f; }
 
 	// Called once by bundle manager after constructing the bundle source
-	virtual void Init(
+	// Any non-fallback errors returned will cause bundle manager to fail to initialize
+	virtual FInstallBundleSourceInitInfo Init(
 		TSharedRef<InstallBundleUtil::FContentRequestStatsMap> InRequestStats,
 		TSharedPtr<IAnalyticsProviderET> AnalyticsProvider,
 		TSharedPtr<InstallBundleUtil::PersistentStats::FPersistentStatContainerBase> PersistentStatsContainer) = 0;
diff --git a/Engine/Source/Runtime/InstallBundleManager/Public/InstallBundleTypes.h b/Engine/Source/Runtime/InstallBundleManager/Public/InstallBundleTypes.h
--- a/Engine/Source/Runtime/InstallBundleManager/Public/InstallBundleTypes.h
+++ b/Engine/Source/Runtime/InstallBundleManager/Public/InstallBundleTypes.h
@@ -153,7 +153,10 @@ struct FInstallBundleSourceInitInfo
 {
 	EInstallBundleManagerInitResult Result = EInstallBundleManagerInitResult::OK;
 	bool bShouldUseFallbackSource = false;
+};
 
+struct FInstallBundleSourceAsyncInitInfo : public FInstallBundleSourceInitInfo
+{
 	TMap<FName, bool> BundleUpToDate;
 };
 
diff --git a/Engine/Source/Runtime/InstallBundleManager/Public/InstallBundleUtils.h b/Engine/Source/Runtime/InstallBundleManager/Public/InstallBundleUtils.h
--- a/Engine/Source/Runtime/InstallBundleManager/Public/InstallBundleUtils.h
+++ b/Engine/Source/Runtime/InstallBundleManager/Public/InstallBundleUtils.h
@@ -167,6 +167,7 @@ namespace InstallBundleUtil
 		void StatsEnd(FName BundleName);
 		void StatsBegin(FName BundleName, const TCHAR* State);
 		void StatsEnd(FName BundleName, const TCHAR* State, uint64 DataSize = 0);
+		void StatsReset(FName BundleName);
 
 		const TMap<FName, InstallBundleUtil::FContentRequestStats>& GetMap() { return StatsMap; }
 	};
