PersistentStats
- Removed some pass-by-value FStrings to use correct const FString& reference instead.
- Added RemoveSessionStats and RemoveBundleStats functions to remove items from the underlying data to clean them up.
#tests iOS Client

diff --git a/Engine/Source/Runtime/InstallBundleManager/Private/InstallBundleUtils.cpp b/Engine/Source/Runtime/InstallBundleManager/Private/InstallBundleUtils.cpp
--- a/Engine/Source/Runtime/InstallBundleManager/Private/InstallBundleUtils.cpp
+++ b/Engine/Source/Runtime/InstallBundleManager/Private/InstallBundleUtils.cpp
@@ -951,6 +951,16 @@ namespace InstallBundleUtil
 			}			
 		}
 
+		void FPersistentStatContainerBase::RemoveSessionStats(const FString& SessionName)
+		{
+			SessionPersistentStatMap.Remove(SessionName);
+		}
+	
+		void FPersistentStatContainerBase::RemoveBundleStats(FName BundleName)
+		{
+			PerBundlePersistentStatMap.Remove(BundleName);
+		}
+	
 		void FPersistentStatContainerBase::StartBundlePersistentStatTracking(FName BundleName, const FString& ExpectedAnalyticsID /* = FString() */, bool bForceResetStatData /* = false */)
 		{
 			//Use the base expected analytics ID if one was not passed in
@@ -984,7 +994,7 @@ namespace InstallBundleUtil
 			}
 		}
 
-		void FPersistentStatContainerBase::StopSessionPersistentStatTracking(FString SessionName, bool bStopAllActiveTimers /* = true */)
+		void FPersistentStatContainerBase::StopSessionPersistentStatTracking(const FString& SessionName, bool bStopAllActiveTimers /* = true */)
 		{
 			FSessionPersistentStats* FoundSessionStats = SessionPersistentStatMap.Find(SessionName);
 			if (nullptr != FoundSessionStats)
@@ -1028,7 +1038,7 @@ namespace InstallBundleUtil
 			OnTimerStartedForStat(FoundBundleStats, TimerToStart);
 		}
 
-		void FPersistentStatContainerBase::StartSessionPersistentStatTimer(FString SessionName, ETimingStatNames TimerToStart)
+		void FPersistentStatContainerBase::StartSessionPersistentStatTimer(const FString& SessionName, ETimingStatNames TimerToStart)
 		{
 			FSessionPersistentStats& FoundSessionStats = SessionPersistentStatMap.FindOrAdd(SessionName, FSessionPersistentStats(SessionName));
 			FoundSessionStats.StartTimingStat(TimerToStart);
@@ -1046,7 +1056,7 @@ namespace InstallBundleUtil
 			OnTimerStoppedForStat(FoundBundleStats, TimerToStop);
 		}
 
-		void FPersistentStatContainerBase::StopSessionPersistentStatTimer(FString SessionName, ETimingStatNames TimerToStop)
+		void FPersistentStatContainerBase::StopSessionPersistentStatTimer(const FString& SessionName, ETimingStatNames TimerToStop)
 		{
 			FSessionPersistentStats& FoundSessionStats = SessionPersistentStatMap.FindOrAdd(SessionName, FSessionPersistentStats(SessionName));
 			FoundSessionStats.StopTimingStat(TimerToStop);
@@ -1060,7 +1070,7 @@ namespace InstallBundleUtil
 			FoundBundleStats.UpdateTimingStat(TimerToUpdate);
 		}
 
-		void FPersistentStatContainerBase::UpdateSessionPersistentStatTimer(FString SessionName, ETimingStatNames TimerToUpdate)
+		void FPersistentStatContainerBase::UpdateSessionPersistentStatTimer(const FString& SessionName, ETimingStatNames TimerToUpdate)
 		{
 			FSessionPersistentStats& FoundSessionStats = SessionPersistentStatMap.FindOrAdd(SessionName, FSessionPersistentStats(SessionName));
 			FoundSessionStats.UpdateTimingStat(TimerToUpdate);
@@ -1074,7 +1084,7 @@ namespace InstallBundleUtil
 			ensureAlwaysMsgf(FoundBundleStats.IsActive(), TEXT("Invalid attempt to increment %s on bundle %s that hasn't yet had StartBundlePersistentStatTracking called on it! Should always start tracking before using persistent stats!"), *LexToString(CounterToUpdate), *(BundleName.ToString()));
 		}
 
-		void FPersistentStatContainerBase::IncrementSessionPersistentCounter(FString SessionName, ECountStatNames CounterToUpdate)
+		void FPersistentStatContainerBase::IncrementSessionPersistentCounter(const FString& SessionName, ECountStatNames CounterToUpdate)
 		{
 			FSessionPersistentStats& FoundSessionStats = SessionPersistentStatMap.FindOrAdd(SessionName, FSessionPersistentStats(SessionName));
 			FoundSessionStats.IncrementCountStat(CounterToUpdate);
diff --git a/Engine/Source/Runtime/InstallBundleManager/Public/InstallBundleUtils.h b/Engine/Source/Runtime/InstallBundleManager/Public/InstallBundleUtils.h
--- a/Engine/Source/Runtime/InstallBundleManager/Public/InstallBundleUtils.h
+++ b/Engine/Source/Runtime/InstallBundleManager/Public/InstallBundleUtils.h
@@ -435,25 +435,29 @@ namespace InstallBundleUtil
 			virtual void StartSessionPersistentStatTracking(const FString& SessionName, const TArray<FName>& RequiredBundles = TArray<FName>(), const FString& ExpectedAnalyticsID = FString(), bool bForceResetStatData = false);
 
 			virtual void StopBundlePersistentStatTracking(FName BundleName, bool bStopAllActiveTimers = true);
-			virtual void StopSessionPersistentStatTracking(FString SessionName, bool bStopAllActiveTimers = true);
+			virtual void StopSessionPersistentStatTracking(const FString& SessionName, bool bStopAllActiveTimers = true);
 
 			virtual void StartBundlePersistentStatTimer(FName BundleName, ETimingStatNames TimerToStart);
-			virtual void StartSessionPersistentStatTimer(FString SessionName, ETimingStatNames TimerToStart);
+			virtual void StartSessionPersistentStatTimer(const FString& SessionName, ETimingStatNames TimerToStart);
 
 			virtual void StopBundlePersistentStatTimer(FName BundleName, ETimingStatNames TimerToStop);
-			virtual void StopSessionPersistentStatTimer(FString SessionName, ETimingStatNames TimerToStop);
+			virtual void StopSessionPersistentStatTimer(const FString& SessionName, ETimingStatNames TimerToStop);
 
 			virtual void UpdateBundlePersistentStatTimer(FName BundleName, ETimingStatNames TimerToUpdate);
-			virtual void UpdateSessionPersistentStatTimer(FString SessionName, ETimingStatNames TimerToUpdate);
+			virtual void UpdateSessionPersistentStatTimer(const FString& SessionName, ETimingStatNames TimerToUpdate);
 
 			virtual void IncrementBundlePersistentCounter(FName BundleName, ECountStatNames CounterToUpdate);
-			virtual void IncrementSessionPersistentCounter(FString SessionName, ECountStatNames CounterToUpdate);
+			virtual void IncrementSessionPersistentCounter(const FString& SessionName, ECountStatNames CounterToUpdate);
 
 			virtual const FBundlePersistentStats* GetBundleStat(FName BundleName) const;
 			virtual const FSessionPersistentStats* GetSessionStat(const FString& SessionName) const;
 
 			virtual void SaveAllDirtyStatsToDisk();
 			
+			//Deletes persistent stat information from our stats for this Session/Bundle to reduce memory useage
+			virtual void RemoveSessionStats(const FString& SessionName);
+			virtual void RemoveBundleStats(FName BundleName);
+			
 		protected:
 			virtual bool Tick(float dt);
 			virtual void ResetTimerUpdate();
