Stub Out IntelligentDelivery install bundle source.
Drive XBox packaging from install bundle config if available.

diff --git a/Engine/Source/Runtime/InstallBundleManager/Private/InstallBundleTypes.cpp b/Engine/Source/Runtime/InstallBundleManager/Private/InstallBundleTypes.cpp
--- a/Engine/Source/Runtime/InstallBundleManager/Private/InstallBundleTypes.cpp
+++ b/Engine/Source/Runtime/InstallBundleManager/Private/InstallBundleTypes.cpp
@@ -13,6 +13,7 @@ const TCHAR* LexToString(EInstallBundleSourceType Type)
 		TEXT("Bulk"),
 		TEXT("BuildPatchServices"),
 		TEXT("PlayGo"),
+		TEXT("IntelligentDelivery"),
 	};
 
 	static_assert(InstallBundleUtil::CastToUnderlying(EInstallBundleSourceType::Count) == UE_ARRAY_COUNT(Strings), "");
diff --git a/Engine/Source/Runtime/InstallBundleManager/Public/InstallBundleTypes.h b/Engine/Source/Runtime/InstallBundleManager/Public/InstallBundleTypes.h
--- a/Engine/Source/Runtime/InstallBundleManager/Public/InstallBundleTypes.h
+++ b/Engine/Source/Runtime/InstallBundleManager/Public/InstallBundleTypes.h
@@ -10,6 +10,7 @@ enum class EInstallBundleSourceType : int
 	Bulk,
 	BuildPatchServices,
 	PlayGo,
+	IntelligentDelivery,
 	Count
 };
 ENUM_RANGE_BY_COUNT(EInstallBundleSourceType, EInstallBundleSourceType::Count);
