- Add MinimumBundleWeight constant in InstallBundleUtil so it isn't duplicated everywhere.

diff --git a/Engine/Source/Runtime/InstallBundleManager/Public/InstallBundleUtils.h b/Engine/Source/Runtime/InstallBundleManager/Public/InstallBundleUtils.h
--- a/Engine/Source/Runtime/InstallBundleManager/Public/InstallBundleUtils.h
+++ b/Engine/Source/Runtime/InstallBundleManager/Public/InstallBundleUtils.h
@@ -19,6 +19,8 @@ namespace InstallBundleUtil
 
 	INSTALLBUNDLEMANAGER_API const FString& GetInstallBundleSectionPrefix();
 
+	constexpr float MinimumBundleWeight = 0.05f;
+
 	// It would really be nice to have these in core
 	template<class EnumType>
 	constexpr auto& CastAsUnderlying(EnumType &Type)
