Add patch checking to install bundle manager interface.

diff --git a/Engine/Source/Runtime/InstallBundleManager/Private/InstallBundleManagerInterface.cpp b/Engine/Source/Runtime/InstallBundleManager/Private/InstallBundleManagerInterface.cpp
--- a/Engine/Source/Runtime/InstallBundleManager/Private/InstallBundleManagerInterface.cpp
+++ b/Engine/Source/Runtime/InstallBundleManager/Private/InstallBundleManagerInterface.cpp
@@ -9,6 +9,8 @@ FInstallBundleCompleteMultiDelegate IInstallBundleManager::RemoveBundleCompleteD
 
 FInstallBundlePausedMultiDelegate IInstallBundleManager::PausedBundleDelegate;
 
+FInstallBundleManagerOnPatchCheckComplete IInstallBundleManager::PatchCheckCompleteDelegate;
+
 IInstallBundleManager* IInstallBundleManager::GetPlatformInstallBundleManager()
 {
 	static IInstallBundleManager* Manager = nullptr;
@@ -25,7 +27,6 @@ IInstallBundleManager* IInstallBundleManager::GetPlatformInstallBundleManager()
 
 		if (FModuleManager::Get().ModuleExists(*ModuleName))
 		{
-			FModuleStatus Status;
 			Module = FModuleManager::LoadModulePtr<IInstallBundleManagerModule>(*ModuleName);
 			if (Module)
 			{
@@ -79,3 +80,8 @@ void IInstallBundleManager::UpdateContentRequestFlags(FName BundleName, EInstall
 	UpdateContentRequestFlags(MakeArrayView(&BundleName, 1), AddFlags, RemoveFlags);
 }
 
+void IInstallBundleManager::StartPatchCheck()
+{
+	PatchCheckCompleteDelegate.Broadcast(EInstallBundleManagerPatchCheckResult::NoPatchRequired);
+}
+
diff --git a/Engine/Source/Runtime/InstallBundleManager/Private/InstallBundleTypes.cpp b/Engine/Source/Runtime/InstallBundleManager/Private/InstallBundleTypes.cpp
--- a/Engine/Source/Runtime/InstallBundleManager/Private/InstallBundleTypes.cpp
+++ b/Engine/Source/Runtime/InstallBundleManager/Private/InstallBundleTypes.cpp
@@ -76,6 +76,7 @@ const TCHAR* LexToString(EInstallBundleResult Result)
 	{
 		TEXT("OK"),
 		TEXT("FailedPrereqRequiresLatestClient"),
+		TEXT("FailedPrereqRequiresLatestContent"),
 		TEXT("InstallError"),
 		TEXT("InstallerOutOfDiskSpaceError"),
 		TEXT("ManifestArchiveError"),
diff --git a/Engine/Source/Runtime/InstallBundleManager/Public/InstallBundleManagerInterface.h b/Engine/Source/Runtime/InstallBundleManager/Public/InstallBundleManagerInterface.h
--- a/Engine/Source/Runtime/InstallBundleManager/Public/InstallBundleManagerInterface.h
+++ b/Engine/Source/Runtime/InstallBundleManager/Public/InstallBundleManagerInterface.h
@@ -53,12 +53,15 @@ DECLARE_DELEGATE_RetVal_OneParam(EInstallBundleManagerInitErrorHandlerResult, FI
 DECLARE_MULTICAST_DELEGATE_OneParam(FInstallBundleCompleteMultiDelegate, FInstallBundleRequestResultInfo);
 DECLARE_MULTICAST_DELEGATE_OneParam(FInstallBundlePausedMultiDelegate, FInstallBundlePauseInfo);
 
+DECLARE_MULTICAST_DELEGATE_OneParam(FInstallBundleManagerOnPatchCheckComplete, EInstallBundleManagerPatchCheckResult);
+
 class INSTALLBUNDLEMANAGER_API IInstallBundleManager
 {
 public:
 	static FInstallBundleCompleteMultiDelegate InstallBundleCompleteDelegate;
 	static FInstallBundleCompleteMultiDelegate RemoveBundleCompleteDelegate;
 	static FInstallBundlePausedMultiDelegate PausedBundleDelegate;
+	static FInstallBundleManagerOnPatchCheckComplete PatchCheckCompleteDelegate;
 
 	static IInstallBundleManager* GetPlatformInstallBundleManager();
 
@@ -103,6 +106,8 @@ public:
 	virtual EInstallBundleRequestFlags GetModifyableContentRequestFlags() const = 0;
 	void UpdateContentRequestFlags(FName BundleName, EInstallBundleRequestFlags AddFlags, EInstallBundleRequestFlags RemoveFlags);
 	virtual void UpdateContentRequestFlags(TArrayView<const FName> BundleNames, EInstallBundleRequestFlags AddFlags, EInstallBundleRequestFlags RemoveFlags) = 0;
+	
+	virtual void StartPatchCheck();
 
 	virtual bool IsNullInterface() const = 0;
 
diff --git a/Engine/Source/Runtime/InstallBundleManager/Public/InstallBundleTypes.h b/Engine/Source/Runtime/InstallBundleManager/Public/InstallBundleTypes.h
--- a/Engine/Source/Runtime/InstallBundleManager/Public/InstallBundleTypes.h
+++ b/Engine/Source/Runtime/InstallBundleManager/Public/InstallBundleTypes.h
@@ -97,6 +97,7 @@ enum class EInstallBundleResult : int
 {
 	OK,
 	FailedPrereqRequiresLatestClient,
+	FailedPrereqRequiresLatestContent,
 	InstallError,
 	InstallerOutOfDiskSpaceError,
 	ManifestArchiveError,
@@ -150,6 +151,21 @@ enum class EInstallBundleStatus : int
 };
 INSTALLBUNDLEMANAGER_API const TCHAR* LexToString(EInstallBundleStatus Status);
 
+enum class EInstallBundleManagerPatchCheckResult : uint32
+{
+	/** No patch required */
+	NoPatchRequired,
+	/** Client Patch required to continue */
+	ClientPatchRequired,
+	/** Content Patch required to continue */
+	ContentPatchRequired,
+	/** Logged in user required for a patch check */
+	NoLoggedInUser,
+	/** Patch check failed */
+	PatchCheckFailure,
+	Count,
+};
+
 struct FInstallBundleSourceInitInfo
 {
 	EInstallBundleManagerInitResult Result = EInstallBundleManagerInitResult::OK;
