Add content patch check to bundle source interface.

diff --git a/Engine/Source/Runtime/InstallBundleManager/Public/InstallBundleSourceInterface.h b/Engine/Source/Runtime/InstallBundleManager/Public/InstallBundleSourceInterface.h
--- a/Engine/Source/Runtime/InstallBundleManager/Public/InstallBundleSourceInterface.h
+++ b/Engine/Source/Runtime/InstallBundleManager/Public/InstallBundleSourceInterface.h
@@ -13,6 +13,8 @@ DECLARE_DELEGATE_TwoParams(FInstallBundleSourceInitDelegate, TSharedRef<IInstall
 DECLARE_DELEGATE_OneParam(FInstallBundleCompleteDelegate, FInstallBundleSourceRequestResultInfo /*Result*/);
 DECLARE_DELEGATE_OneParam(FInstallBundlePausedDelegate, FInstallBundleSourcePauseInfo /*PauseInfo*/);
 
+DECLARE_DELEGATE_TwoParams(FInstallBundleSourceContentPatchResultDelegate, TSharedRef<IInstallBundleSource> /*Source*/, bool /*bContentPatchRequired*/);
+
 class IInstallBundleSource : public TSharedFromThis<IInstallBundleSource>
 {
 public:
@@ -84,6 +86,8 @@ public:
 	// Derived classes should implement this if their content install will take a significant amount of time
 	virtual TOptional<FInstallBundleSourceProgress> GetBundleProgress(FName BundleName) const { return TOptional<FInstallBundleSourceProgress>(); }
 
+	virtual void CheckForContentPatch(FInstallBundleSourceContentPatchResultDelegate Callback) { Callback.Execute(AsShared(), false); }
+
 	// Called by bundle manager to pass through command line options to simulate errors
 	virtual void SetErrorSimulationCommands(const FString& CommandLine) {}
 };
