Add EInstallBundlePriority that maps to background download priority.

diff --git a/Engine/Source/Runtime/InstallBundleManager/Private/InstallBundleTypes.cpp b/Engine/Source/Runtime/InstallBundleManager/Private/InstallBundleTypes.cpp
--- a/Engine/Source/Runtime/InstallBundleManager/Private/InstallBundleTypes.cpp
+++ b/Engine/Source/Runtime/InstallBundleManager/Private/InstallBundleTypes.cpp
@@ -102,6 +102,56 @@ const TCHAR* LexToString(EInstallBundleStatus Status)
 	return Strings[InstallBundleUtil::CastToUnderlying(Status)];
 }
 
+const TCHAR* LexToString(EInstallBundleManagerPatchCheckResult EnumVal)
+{
+	// These are namespaced because PartyHub expects them that way :/
+	static const TCHAR* Strings[] =
+	{
+		TEXT("EInstallBundleManagerPatchCheckResult::NoPatchRequired"),
+		TEXT("EInstallBundleManagerPatchCheckResult::ClientPatchRequired"),
+		TEXT("EInstallBundleManagerPatchCheckResult::ContentPatchRequired"),
+		TEXT("EInstallBundleManagerPatchCheckResult::NoLoggedInUser"),
+		TEXT("EInstallBundleManagerPatchCheckResult::PatchCheckFailure"),
+	};
+
+	static_assert(InstallBundleUtil::CastToUnderlying(EInstallBundleManagerPatchCheckResult::Count) == UE_ARRAY_COUNT(Strings), "");
+	return Strings[InstallBundleUtil::CastToUnderlying(EnumVal)];
+}
+
+const TCHAR* LexToString(EInstallBundlePriority Priority)
+{
+	static const TCHAR* Strings[] =
+	{
+		TEXT("High"),
+		TEXT("Normal"),
+		TEXT("Low"),
+	};
+
+	static_assert(InstallBundleUtil::CastToUnderlying(EInstallBundlePriority::Count) == UE_ARRAY_COUNT(Strings), "");
+	return Strings[InstallBundleUtil::CastToUnderlying(Priority)];
+}
+
+bool LexTryParseString(EInstallBundlePriority& OutMode, const TCHAR* InBuffer)
+{
+	if (FCString::Stricmp(InBuffer, TEXT("High")) == 0)
+	{
+		OutMode = EInstallBundlePriority::High;
+		return true;
+	}
+	if (FCString::Stricmp(InBuffer, TEXT("Normal")) == 0)
+	{
+		OutMode = EInstallBundlePriority::Normal;
+		return true;
+	}
+	if (FCString::Stricmp(InBuffer, TEXT("Low")) == 0)
+	{
+		OutMode = EInstallBundlePriority::Low;
+		return true;
+	}
+	return false;
+}
+
+
 bool FInstallBundleCombinedContentState::GetAllBundlesHaveState(EInstallBundleContentState State, TArrayView<const FName> ExcludedBundles /*= TArrayView<const FName>()*/) const
 {
 	for (const TPair<FName, FInstallBundleContentState>& Pair : IndividualBundleStates)
diff --git a/Engine/Source/Runtime/InstallBundleManager/Public/InstallBundleTypes.h b/Engine/Source/Runtime/InstallBundleManager/Public/InstallBundleTypes.h
--- a/Engine/Source/Runtime/InstallBundleManager/Public/InstallBundleTypes.h
+++ b/Engine/Source/Runtime/InstallBundleManager/Public/InstallBundleTypes.h
@@ -165,20 +165,21 @@ enum class EInstallBundleManagerPatchCheckResult : uint32
 	PatchCheckFailure,
 	Count,
 };
+INSTALLBUNDLEMANAGER_API const TCHAR* LexToString(EInstallBundleManagerPatchCheckResult EnumVal);
 
-inline const TCHAR* LexToString(EInstallBundleManagerPatchCheckResult EnumVal)
+/**
+ * Enum used to describe download priority. Higher priorities will be downloaded first.
+ * Note: Should always be kept in High -> Low priority order if adding more Priorities!
+ */
+enum class EInstallBundlePriority : uint8
 {
-	static_assert((int32)EInstallBundleManagerPatchCheckResult::Count == 5, "Please add support for the extra values to the function below.");
-	switch (EnumVal)
-	{
-		case EInstallBundleManagerPatchCheckResult::NoPatchRequired: return TEXT("EInstallBundleManagerPatchCheckResult::NoPatchRequired");
-		case EInstallBundleManagerPatchCheckResult::ClientPatchRequired: return TEXT("EInstallBundleManagerPatchCheckResult::ClientPatchRequired");
-		case EInstallBundleManagerPatchCheckResult::ContentPatchRequired: return TEXT("EInstallBundleManagerPatchCheckResult::ContentPatchRequired");
-		case EInstallBundleManagerPatchCheckResult::NoLoggedInUser: return TEXT("EInstallBundleManagerPatchCheckResult::NoLoggedInUser");
-		case EInstallBundleManagerPatchCheckResult::PatchCheckFailure: return TEXT("EInstallBundleManagerPatchCheckResult::PatchCheckFailure");
-		default: return TEXT("EInstallBundleManagerPatchCheckResult::Unknown");
-	}
+	High,
+	Normal,
+	Low,
+	Count
 };
+INSTALLBUNDLEMANAGER_API const TCHAR* LexToString(EInstallBundlePriority Priority);
+INSTALLBUNDLEMANAGER_API bool LexTryParseString(EInstallBundlePriority& OutMode, const TCHAR* InBuffer);
 
 struct FInstallBundleSourceInitInfo
 {
