Bundle Manager Init no longer directly reads bundle config during Init.  Still need to fix dependencies.

diff --git a/Engine/Source/Runtime/InstallBundleManager/Public/InstallBundleSourceInterface.h b/Engine/Source/Runtime/InstallBundleManager/Public/InstallBundleSourceInterface.h
--- a/Engine/Source/Runtime/InstallBundleManager/Public/InstallBundleSourceInterface.h
+++ b/Engine/Source/Runtime/InstallBundleManager/Public/InstallBundleSourceInterface.h
@@ -10,6 +10,8 @@ class IAnalyticsProviderET;
 
 DECLARE_DELEGATE_TwoParams(FInstallBundleSourceInitDelegate, TSharedRef<IInstallBundleSource> /*Source*/, FInstallBundleSourceAsyncInitInfo /*InitInfo*/);
 
+DECLARE_DELEGATE_TwoParams(FInstallBundleSourceQueryBundleInfoDelegate, TSharedRef<IInstallBundleSource> /*Source*/, FInstallBundleSourceBundleInfoQueryResultInfo /*Result*/);
+
 DECLARE_DELEGATE_OneParam(FInstallBundleCompleteDelegate, FInstallBundleSourceRequestResultInfo /*Result*/);
 DECLARE_DELEGATE_OneParam(FInstallBundlePausedDelegate, FInstallBundleSourcePauseInfo /*PauseInfo*/);
 
@@ -36,6 +38,8 @@ public:
 	// It will be retried indefinitely until init is successful.  
 	virtual void AsyncInit(FInstallBundleSourceInitDelegate Callback) = 0;
 
+	virtual void QueryBundleInfo(FInstallBundleSourceQueryBundleInfoDelegate Callback) = 0; /*{}*/
+
 	// Whether this source has been initialized or not
 	virtual EInstallBundleManagerInitState GetInitState() const = 0;
 
diff --git a/Engine/Source/Runtime/InstallBundleManager/Public/InstallBundleTypes.h b/Engine/Source/Runtime/InstallBundleManager/Public/InstallBundleTypes.h
--- a/Engine/Source/Runtime/InstallBundleManager/Public/InstallBundleTypes.h
+++ b/Engine/Source/Runtime/InstallBundleManager/Public/InstallBundleTypes.h
@@ -189,7 +189,22 @@ struct FInstallBundleSourceInitInfo
 
 struct FInstallBundleSourceAsyncInitInfo : public FInstallBundleSourceInitInfo
 {
-	TMap<FName, bool> BundleUpToDate;
+	// Reserved for future use
+};
+
+struct FInstallBundleSourceBundleInfo
+{
+	FName BundleName;
+	FString BundleNameString;
+	EInstallBundlePriority Priority = EInstallBundlePriority::Low;
+	bool bIsStartup = false; // Only one startup bundle allowed.  All sources must agree on this.
+	bool bDoPatchCheck = false; // This bundle should do a patch check and fail if it doesn't pass
+	bool bBundleUpToDate = false; // Whether this bundle is up to date
+};
+
+struct FInstallBundleSourceBundleInfoQueryResultInfo
+{
+	TMap<FName, FInstallBundleSourceBundleInfo> SourceBundleInfoMap;
 };
 
 struct FInstallBundleSourceRequestResultInfo
