Adding new InstallBundleType GameCustom to allow for custom game mode install bundle sources.

diff --git a/Engine/Source/Runtime/InstallBundleManager/Private/InstallBundleTypes.cpp b/Engine/Source/Runtime/InstallBundleManager/Private/InstallBundleTypes.cpp
--- a/Engine/Source/Runtime/InstallBundleManager/Private/InstallBundleTypes.cpp
+++ b/Engine/Source/Runtime/InstallBundleManager/Private/InstallBundleTypes.cpp
@@ -14,6 +14,7 @@ const TCHAR* LexToString(EInstallBundleSourceType Type)
 		TEXT("BuildPatchServices"),
 		TEXT("PlayGo"),
 		TEXT("IntelligentDelivery"),
+		TEXT("GameCustom"),
 	};
 
 	static_assert(InstallBundleUtil::CastToUnderlying(EInstallBundleSourceType::Count) == UE_ARRAY_COUNT(Strings), "");
diff --git a/Engine/Source/Runtime/InstallBundleManager/Public/InstallBundleTypes.h b/Engine/Source/Runtime/InstallBundleManager/Public/InstallBundleTypes.h
--- a/Engine/Source/Runtime/InstallBundleManager/Public/InstallBundleTypes.h
+++ b/Engine/Source/Runtime/InstallBundleManager/Public/InstallBundleTypes.h
@@ -11,7 +11,8 @@ enum class EInstallBundleSourceType : int
 	BuildPatchServices,
 	PlayGo,
 	IntelligentDelivery,
-	Count
+	GameCustom,
+	Count,
 };
 ENUM_RANGE_BY_COUNT(EInstallBundleSourceType, EInstallBundleSourceType::Count);
 INSTALLBUNDLEMANAGER_API const TCHAR* LexToString(EInstallBundleSourceType Type);
