Install Bundle Sources now provide bunld edependencies to bundle manager.

diff --git a/Engine/Source/Runtime/InstallBundleManager/Public/InstallBundleSourceInterface.h b/Engine/Source/Runtime/InstallBundleManager/Public/InstallBundleSourceInterface.h
--- a/Engine/Source/Runtime/InstallBundleManager/Public/InstallBundleSourceInterface.h
+++ b/Engine/Source/Runtime/InstallBundleManager/Public/InstallBundleSourceInterface.h
@@ -38,7 +38,9 @@ public:
 	// It will be retried indefinitely until init is successful.  
 	virtual void AsyncInit(FInstallBundleSourceInitDelegate Callback) = 0;
 
-	virtual void QueryBundleInfo(FInstallBundleSourceQueryBundleInfoDelegate Callback) = 0; /*{}*/
+	// Currently only called after AsyncInit initialization.
+	// Provides information about bundles this source knows about back to bundle manager.
+	virtual void QueryBundleInfo(FInstallBundleSourceQueryBundleInfoDelegate Callback) = 0;
 
 	// Whether this source has been initialized or not
 	virtual EInstallBundleManagerInitState GetInitState() const = 0;
@@ -46,6 +48,10 @@ public:
 	// Returns content version in a "<BuildVersion>-<Platform>" format
 	virtual FString GetContentVersion() const = 0;
 
+	// Finds all dependencies for InBundleName, including InBundleName
+	// Sets bSkippedUnknownBundles if information for InBundleName or a dependency can't be found
+	virtual TSet<FName> GetBundleDependencies(FName InBundleName, bool* bSkippedUnknownBundles /*= nullptr*/) const = 0;
+
 	// Gets the state of content on disk
 	// BundleNames contains all dependencies and has been deduped
 	virtual void GetContentState(TArrayView<const FName> BundleNames, EInstallBundleGetContentStateFlags Flags, FInstallBundleGetContentStateDelegate Callback) = 0;
