Provide a way for bundle sources to add/update bundle info after init.

diff --git a/Engine/Source/Runtime/InstallBundleManager/Public/InstallBundleSourceInterface.h b/Engine/Source/Runtime/InstallBundleManager/Public/InstallBundleSourceInterface.h
--- a/Engine/Source/Runtime/InstallBundleManager/Public/InstallBundleSourceInterface.h
+++ b/Engine/Source/Runtime/InstallBundleManager/Public/InstallBundleSourceInterface.h
@@ -11,6 +11,7 @@ class IAnalyticsProviderET;
 DECLARE_DELEGATE_TwoParams(FInstallBundleSourceInitDelegate, TSharedRef<IInstallBundleSource> /*Source*/, FInstallBundleSourceAsyncInitInfo /*InitInfo*/);
 
 DECLARE_DELEGATE_TwoParams(FInstallBundleSourceQueryBundleInfoDelegate, TSharedRef<IInstallBundleSource> /*Source*/, FInstallBundleSourceBundleInfoQueryResultInfo /*Result*/);
+DECLARE_DELEGATE_RetVal_TwoParams(EInstallBundleSourceUpdateBundleInfoResult, FInstallBundleSourceUpdateBundleInfoDelegate, TSharedRef<IInstallBundleSource> /*Source*/, FInstallBundleSourceBundleInfoQueryResultInfo /*Result*/);
 
 DECLARE_DELEGATE_OneParam(FInstallBundleCompleteDelegate, FInstallBundleSourceRequestResultInfo /*Result*/);
 DECLARE_DELEGATE_OneParam(FInstallBundlePausedDelegate, FInstallBundleSourcePauseInfo /*PauseInfo*/);
@@ -40,7 +41,9 @@ public:
 
 	// Currently only called after AsyncInit initialization.
 	// Provides information about bundles this source knows about back to bundle manager.
-	virtual void QueryBundleInfo(FInstallBundleSourceQueryBundleInfoDelegate Callback) = 0;
+	virtual void AsyncInit_QueryBundleInfo(FInstallBundleSourceQueryBundleInfoDelegate OnCompleteCallback) = 0;
+
+	virtual void AsyncInit_SetUpdateBundleInfoCallback(FInstallBundleSourceUpdateBundleInfoDelegate UpdateCallback) {}
 
 	// Whether this source has been initialized or not
 	virtual EInstallBundleManagerInitState GetInitState() const = 0;
diff --git a/Engine/Source/Runtime/InstallBundleManager/Public/InstallBundleTypes.h b/Engine/Source/Runtime/InstallBundleManager/Public/InstallBundleTypes.h
--- a/Engine/Source/Runtime/InstallBundleManager/Public/InstallBundleTypes.h
+++ b/Engine/Source/Runtime/InstallBundleManager/Public/InstallBundleTypes.h
@@ -208,6 +208,15 @@ struct FInstallBundleSourceBundleInfoQueryResultInfo
 	TMap<FName, FInstallBundleSourceBundleInfo> SourceBundleInfoMap;
 };
 
+enum class EInstallBundleSourceUpdateBundleInfoResult : uint32
+{
+	OK,
+	AlreadyMounted,
+	AlreadyRequested,
+	IllegalStartupBundle,
+	Count,
+};
+
 struct FInstallBundleSourceRequestResultInfo
 {
 	FName BundleName;
