Add an install bundle request flag to skip mounting.
Add an additional callback for when a bundle is done being installed/patched.

diff --git a/Engine/Source/Runtime/InstallBundleManager/Private/InstallBundleManagerInterface.cpp b/Engine/Source/Runtime/InstallBundleManager/Private/InstallBundleManagerInterface.cpp
--- a/Engine/Source/Runtime/InstallBundleManager/Private/InstallBundleManagerInterface.cpp
+++ b/Engine/Source/Runtime/InstallBundleManager/Private/InstallBundleManagerInterface.cpp
@@ -3,9 +3,9 @@
 #include "InstallBundleManagerInterface.h"
 #include "InstallBundleManagerModule.h"
 
-FInstallBundleCompleteMultiDelegate IInstallBundleManager::InstallBundleCompleteDelegate;
+FInstallBundleCompleteMultiDelegate IInstallBundleManager::InstallBundleUpdatedDelegate;
 
-FInstallBundleCompleteMultiDelegate IInstallBundleManager::RemoveBundleCompleteDelegate;
+FInstallBundleCompleteMultiDelegate IInstallBundleManager::InstallBundleCompleteDelegate;
 
 FInstallBundlePausedMultiDelegate IInstallBundleManager::PausedBundleDelegate;
 
diff --git a/Engine/Source/Runtime/InstallBundleManager/Public/InstallBundleManagerInterface.h b/Engine/Source/Runtime/InstallBundleManager/Public/InstallBundleManagerInterface.h
--- a/Engine/Source/Runtime/InstallBundleManager/Public/InstallBundleManagerInterface.h
+++ b/Engine/Source/Runtime/InstallBundleManager/Public/InstallBundleManagerInterface.h
@@ -60,8 +60,8 @@ DECLARE_DELEGATE_RetVal(bool, FInstallBundleManagerEnvironmentWantsPatchCheck);
 class INSTALLBUNDLEMANAGER_API IInstallBundleManager
 {
 public:
-	static FInstallBundleCompleteMultiDelegate InstallBundleCompleteDelegate;
-	static FInstallBundleCompleteMultiDelegate RemoveBundleCompleteDelegate;
+	static FInstallBundleCompleteMultiDelegate InstallBundleUpdatedDelegate;  // Called when content is up to do date
+	static FInstallBundleCompleteMultiDelegate InstallBundleCompleteDelegate; // Called when content is ready to use
 	static FInstallBundlePausedMultiDelegate PausedBundleDelegate;
 	static FInstallBundleManagerOnPatchCheckComplete PatchCheckCompleteDelegate;
 
diff --git a/Engine/Source/Runtime/InstallBundleManager/Public/InstallBundleTypes.h b/Engine/Source/Runtime/InstallBundleManager/Public/InstallBundleTypes.h
--- a/Engine/Source/Runtime/InstallBundleManager/Public/InstallBundleTypes.h
+++ b/Engine/Source/Runtime/InstallBundleManager/Public/InstallBundleTypes.h
@@ -86,11 +86,12 @@ enum class EInstallBundleRequestInfoFlags : int32
 	None = 0,
 	EnqueuedBundlesForInstall = (1 << 0),
 	SkippedAlreadyMountedBundles = (1 << 1),
-	SkippedUnknownBundles = (1 << 2),
-	SkippedInvalidBundles = (1 << 3), // Bundle can't be used with this build
-	SkippedUnusableLanguageBundles = (1 << 4), // Can't enqueue language bundles because of current system settings
-	SkippedBundlesDueToBundleSource = (1 << 5), // A bundle source rejected a bundle for some reason
-	InitializationError = (1 << 6), // Can't enqueue because the bundle manager failed to initialize
+	SkippedAlreadyUpdatedBundles = (1 << 2), // Only possible with EInstallBundleRequestFlags::SkipMount
+	SkippedUnknownBundles = (1 << 3),
+	SkippedInvalidBundles = (1 << 4), // Bundle can't be used with this build
+	SkippedUnusableLanguageBundles = (1 << 5), // Can't enqueue language bundles because of current system settings
+	SkippedBundlesDueToBundleSource = (1 << 6), // A bundle source rejected a bundle for some reason
+	InitializationError = (1 << 7), // Can't enqueue because the bundle manager failed to initialize
 };
 ENUM_CLASS_FLAGS(EInstallBundleRequestInfoFlags);
 
@@ -116,6 +117,7 @@ enum class EInstallBundleRequestFlags : uint32
 	SendNotificationIfDownloadCompletesInBackground = (1 << 2),
 	ForceNoPatching = (1 << 3),
 	TrackPersistentBundleStats = (1 << 4),
+	SkipMount = (1 << 5),
 	Defaults = UseBackgroundDownloads,
 };
 ENUM_CLASS_FLAGS(EInstallBundleRequestFlags)
