Initial install bundle cache work - checking in disabled.
Stub out install bundle removal API for install bundle sources.
Make install bundle sources gracefully handle bundles that are unkown to them.

diff --git a/Engine/Source/Runtime/InstallBundleManager/Private/InstallBundleCache.cpp b/Engine/Source/Runtime/InstallBundleManager/Private/InstallBundleCache.cpp
new file mode 100644
--- /dev/null
+++ b/Engine/Source/Runtime/InstallBundleManager/Private/InstallBundleCache.cpp
@@ -0,0 +1,196 @@
+// Copyright Epic Games, Inc. All Rights Reserved.
+
+#include "InstallBundleCache.h"
+#include "InstallBundleManagerPrivatePCH.h"
+
+#define INSTALLBUNDLE_CACHE_CHECK_INVARIANTS (DO_CHECK && 1)
+
+FInstallBundleCache::~FInstallBundleCache()
+{
+}
+
+void FInstallBundleCache::Init(FInstallBundleCacheInitInfo InitInfo)
+{
+	CacheName = InitInfo.CacheName;
+	TotalSize = InitInfo.Size;
+	DeleteBundleFiles = MoveTemp(InitInfo.DeleteBundleFiles);
+
+	check(DeleteBundleFiles.IsBound());
+}
+
+void FInstallBundleCache::AddOrUpdateBundle(const FInstallBundleCacheAddBundleInfo& AddInfo)
+{
+	FBundleCacheInfo& BundleCacheInfo = CacheInfo.FindOrAdd(AddInfo.BundleName);
+	BundleCacheInfo.FullInstallSize = AddInfo.FullInstallSize;
+	BundleCacheInfo.CurrentInstallSize = AddInfo.FullInstallSize;
+
+	CheckInvariants();
+}
+
+void FInstallBundleCache::RemoveBundle(FName BundleName)
+{
+	CacheInfo.Remove(BundleName);
+}
+
+bool FInstallBundleCache::HasBundle(FName BundleName)
+{
+	return CacheInfo.Contains(BundleName);
+}
+
+uint64 FInstallBundleCache::GetSize() const
+{
+	return TotalSize;
+}
+
+uint64 FInstallBundleCache::GetFreeSpace() const
+{
+	uint64 UsedSize = 0;
+	for (const TPair<FName, FBundleCacheInfo>& Pair : CacheInfo)
+	{
+		UsedSize += Pair.Value.GetSize();
+	}
+
+	if (UsedSize > TotalSize)
+		return 0;
+
+	return TotalSize - UsedSize;
+}
+
+FInstallBundleCacheReserveResult FInstallBundleCache::Reserve(FName BundleName)
+{
+	FInstallBundleCacheReserveResult Result;
+
+	FBundleCacheInfo* BundleInfo = CacheInfo.Find(BundleName);
+	if (BundleInfo == nullptr)
+	{
+		Result.Result = EInstallBundleCacheReserveResult::Success;
+		return Result;
+	}
+
+	if (BundleInfo->State == ECacheState::PendingEvict)
+	{
+		Result.Result = EInstallBundleCacheReserveResult::Failure;
+		return Result;
+	}
+
+	if (BundleInfo->State == ECacheState::Reserved)
+	{
+		Result.Result = EInstallBundleCacheReserveResult::Success;
+		return Result;
+	}
+
+	if (BundleInfo->FullInstallSize <= BundleInfo->CurrentInstallSize)
+	{
+		BundleInfo->State = ECacheState::Reserved;
+		Result.Result = EInstallBundleCacheReserveResult::Success;
+		return Result;
+	}
+
+	uint64 SizeNeeded = BundleInfo->FullInstallSize - BundleInfo->CurrentInstallSize;
+	uint64 FreeSpace = GetFreeSpace();
+	if (FreeSpace >= SizeNeeded)
+	{
+		BundleInfo->State = ECacheState::Reserved;
+		Result.Result = EInstallBundleCacheReserveResult::Success;
+		return Result;
+	}
+
+	Result.Result = EInstallBundleCacheReserveResult::NeedsEvict;
+
+	// TODO: Should search in LRU order
+	// TODO: LRU sort should consider bHintReqeusted
+
+	uint64 CanFreeSpace = FreeSpace;
+	for (const TPair<FName, FBundleCacheInfo>& Pair : CacheInfo)
+	{
+		if(Pair.Key == BundleName)
+			continue;
+
+		if(Pair.Value.State == ECacheState::Reserved)
+			continue;
+
+		uint64 BundleSize = Pair.Value.GetSize();
+		if (BundleSize > 0)
+		{
+			CanFreeSpace += Pair.Value.GetSize();
+			Result.BundlesToEvict.Add(Pair.Key);
+		}
+
+		if(CanFreeSpace >= SizeNeeded)
+			break;
+	}
+
+	if (CanFreeSpace < SizeNeeded)
+	{
+		Result.Result = EInstallBundleCacheReserveResult::Failure;
+	}
+	else
+	{
+		check(Result.BundlesToEvict.Num() > 0);
+	}
+
+	return Result;
+}
+
+bool FInstallBundleCache::Release(FName BundleName)
+{
+	FBundleCacheInfo* BundleInfo = CacheInfo.Find(BundleName);
+	if (BundleInfo == nullptr)
+	{
+		return true;
+	}
+
+	if (BundleInfo->State == ECacheState::Released)
+	{
+		return true;
+	}
+
+	if (BundleInfo->State == ECacheState::Reserved || 
+		BundleInfo->State == ECacheState::PendingEvict)
+	{
+		BundleInfo->State = ECacheState::Released;
+		return true;
+	}
+
+	return false;
+}
+
+bool FInstallBundleCache::SetPendingEvict(FName BundleName)
+{
+	FBundleCacheInfo* BundleInfo = CacheInfo.Find(BundleName);
+	if (BundleInfo == nullptr)
+	{
+		return true;
+	}
+
+	if (BundleInfo->State == ECacheState::PendingEvict)
+	{
+		return true;
+	}
+
+	if (BundleInfo->State == ECacheState::Released)
+	{
+		BundleInfo->State = ECacheState::PendingEvict;
+		return true;
+	}
+
+	return false;
+}
+
+void FInstallBundleCache::HintRequested(FName BundleName, bool bRequested)
+{
+	FBundleCacheInfo* BundleInfo = CacheInfo.Find(BundleName);
+	if (BundleInfo)
+	{
+		BundleInfo->bHintReqeusted = bRequested;
+	}
+}
+
+void FInstallBundleCache::CheckInvariants() const
+{
+#if INSTALLBUNDLE_CACHE_CHECK_INVARIANTS
+
+	// TODO
+
+#endif // INSTALLBUNDLE_CACHE_CHECK_INVARIANTS
+}
diff --git a/Engine/Source/Runtime/InstallBundleManager/Private/InstallBundleTypes.cpp b/Engine/Source/Runtime/InstallBundleManager/Private/InstallBundleTypes.cpp
--- a/Engine/Source/Runtime/InstallBundleManager/Private/InstallBundleTypes.cpp
+++ b/Engine/Source/Runtime/InstallBundleManager/Private/InstallBundleTypes.cpp
@@ -78,6 +78,7 @@ const TCHAR* LexToString(EInstallBundleResult Result)
 		TEXT("OK"),
 		TEXT("FailedPrereqRequiresLatestClient"),
 		TEXT("FailedPrereqRequiresLatestContent"),
+		TEXT("FailedCacheReserve"),
 		TEXT("InstallError"),
 		TEXT("InstallerOutOfDiskSpaceError"),
 		TEXT("ManifestArchiveError"),
diff --git a/Engine/Source/Runtime/InstallBundleManager/Public/InstallBundleCache.h b/Engine/Source/Runtime/InstallBundleManager/Public/InstallBundleCache.h
new file mode 100644
--- /dev/null
+++ b/Engine/Source/Runtime/InstallBundleManager/Public/InstallBundleCache.h
@@ -0,0 +1,112 @@
+// Copyright Epic Games, Inc. All Rights Reserved.
+
+#pragma once
+
+#include "InstallBundleTypes.h"
+
+class FInstallBundleCache;
+
+DECLARE_DELEGATE_TwoParams(FInstallBundleCacheEvictDelegate, TSharedRef<FInstallBundleCache> /*Cache*/, FName /*BundleName*/);
+
+struct FInstallBundleCacheInitInfo
+{
+	FName CacheName;
+	uint64 Size = 0;
+	FInstallBundleCacheEvictDelegate DeleteBundleFiles;
+};
+
+struct FInstallBundleCacheAddBundleInfo
+{
+	FName BundleName;
+	uint64 FullInstallSize = 0; // Total disk footprint when this bundle is fully installed
+	uint64 CurrentInstallSize = 0; // Disk footprint of the bundle in it's current state
+};
+
+enum class EInstallBundleCacheReserveResult : int8
+{
+	Failure,
+	Success,
+	NeedsEvict,
+};
+
+struct FInstallBundleCacheReserveResult
+{
+	TArray<FName> BundlesToEvict;
+	EInstallBundleCacheReserveResult Result = EInstallBundleCacheReserveResult::Failure;
+};
+
+class INSTALLBUNDLEMANAGER_API FInstallBundleCache : public TSharedFromThis<FInstallBundleCache>
+{
+public:
+	virtual ~FInstallBundleCache();
+
+	void Init(FInstallBundleCacheInitInfo InitInfo);
+
+	// Add a bundle to the cache.  
+	void AddOrUpdateBundle(const FInstallBundleCacheAddBundleInfo& AddInfo);
+
+	void RemoveBundle(FName BundleName);
+
+	bool HasBundle(FName BundleName);
+
+	// Return the total size of the cache
+	uint64 GetSize() const;
+	// Return the amount of free space in the cache
+	uint64 GetFreeSpace() const;
+
+	// Called from bundle manager
+	FInstallBundleCacheReserveResult Reserve(FName BundleName);
+
+	// Called from bundle manager to make the files for this bundle eligible for eviction
+	bool Release(FName BundleName);
+
+	bool SetPendingEvict(FName BundleName);
+
+	// Hint to the cache that this bundle is requested, and we should prefer to evict non-requested bundles if possible
+	void HintRequested(FName BundleName, bool bRequested);
+
+private:
+	void CheckInvariants() const;
+
+private:
+	enum class ECacheState : uint8
+	{
+		Released, //Transitions to Reserved or PendingEvict
+		Reserved, //Transitions to Released
+		PendingEvict, // Transitions to Released
+	};
+
+	struct FBundleCacheInfo 
+	{
+		uint64 FullInstallSize = 0;
+		uint64 CurrentInstallSize = 0;
+		ECacheState State = ECacheState::Released;
+		bool bHintReqeusted = false; // Hint to the cache that this bundle is requested, and we should prefer to evict non-requested bundles if possible
+
+		uint64 GetSize() const
+		{
+			if (State == ECacheState::Released)
+				return CurrentInstallSize;
+
+			// Just consider any pending evictions to be 0 size.
+			// We will flush any required evictions before installing
+			// bundles that use this cache.
+			if (State == ECacheState::PendingEvict)
+				return 0;
+
+			if (CurrentInstallSize > FullInstallSize)
+				return CurrentInstallSize;
+
+			return FullInstallSize;
+		}
+	};
+
+private:
+	FInstallBundleCacheEvictDelegate DeleteBundleFiles;
+
+	TMap<FName, FBundleCacheInfo> CacheInfo;
+
+	uint64 TotalSize = 0;
+
+	FName CacheName;
+};
diff --git a/Engine/Source/Runtime/InstallBundleManager/Public/InstallBundleSourceInterface.h b/Engine/Source/Runtime/InstallBundleManager/Public/InstallBundleSourceInterface.h
--- a/Engine/Source/Runtime/InstallBundleManager/Public/InstallBundleSourceInterface.h
+++ b/Engine/Source/Runtime/InstallBundleManager/Public/InstallBundleSourceInterface.h
@@ -13,8 +13,9 @@ DECLARE_DELEGATE_TwoParams(FInstallBundleSourceInitDelegate, TSharedRef<IInstall
 DECLARE_DELEGATE_TwoParams(FInstallBundleSourceQueryBundleInfoDelegate, TSharedRef<IInstallBundleSource> /*Source*/, FInstallBundleSourceBundleInfoQueryResultInfo /*Result*/);
 DECLARE_DELEGATE_RetVal_TwoParams(EInstallBundleSourceUpdateBundleInfoResult, FInstallBundleSourceUpdateBundleInfoDelegate, TSharedRef<IInstallBundleSource> /*Source*/, FInstallBundleSourceBundleInfoQueryResultInfo /*Result*/);
 
-DECLARE_DELEGATE_OneParam(FInstallBundleCompleteDelegate, FInstallBundleSourceRequestResultInfo /*Result*/);
+DECLARE_DELEGATE_OneParam(FInstallBundleCompleteDelegate, FInstallBundleSourceUpdateContentResultInfo /*Result*/);
 DECLARE_DELEGATE_OneParam(FInstallBundlePausedDelegate, FInstallBundleSourcePauseInfo /*PauseInfo*/);
+DECLARE_DELEGATE_OneParam(FInstallBundleRemovedDelegate, FInstallBundleSourceRemoveContentResultInfo /*Result*/);
 
 DECLARE_DELEGATE_TwoParams(FInstallBundleSourceContentPatchResultDelegate, TSharedRef<IInstallBundleSource> /*Source*/, bool /*bContentPatchRequired*/);
 
@@ -77,6 +78,22 @@ public:
 	// BundleContexts contains all dependencies and has been deduped
 	virtual void RequestUpdateContent(FRequestUpdateContentBundleContext BundleContext) = 0;
 
+	struct FRequestRemoveContentBundleContext
+	{
+		FName BundleName;
+		FInstallBundleRemovedDelegate CompleteCallback;
+	};
+
+	// Removes content from disk if present
+	// BundleContexts contains all dependencies and has been deduped
+	// Bundle manager will not schedule removes at the same time as updates for the same bundle
+	virtual void RequestRemoveContent(FRequestRemoveContentBundleContext BundleContext)
+	{
+		FInstallBundleSourceRemoveContentResultInfo ResultInfo;
+		ResultInfo.BundleName = BundleContext.BundleName;
+		BundleContext.CompleteCallback.Execute(MoveTemp(ResultInfo));
+	}
+
 	// Returns true if content is scheduled to be removed the next time the source is initialized
 	// BundleNames contains all dependencies and has been deduped
 	virtual bool RequestRemoveContentOnNextInit(TArrayView<const FName> RemoveNames) { return false; }
diff --git a/Engine/Source/Runtime/InstallBundleManager/Public/InstallBundleTypes.h b/Engine/Source/Runtime/InstallBundleManager/Public/InstallBundleTypes.h
--- a/Engine/Source/Runtime/InstallBundleManager/Public/InstallBundleTypes.h
+++ b/Engine/Source/Runtime/InstallBundleManager/Public/InstallBundleTypes.h
@@ -100,6 +100,7 @@ enum class EInstallBundleResult : int
 	OK,
 	FailedPrereqRequiresLatestClient,
 	FailedPrereqRequiresLatestContent,
+	FailedCacheReserve,
 	InstallError,
 	InstallerOutOfDiskSpaceError,
 	ManifestArchiveError,
@@ -200,9 +201,12 @@ struct FInstallBundleSourceBundleInfo
 	FName BundleName;
 	FString BundleNameString;
 	EInstallBundlePriority Priority = EInstallBundlePriority::Low;
+	uint64 FullInstallSize = 0; // Total disk footprint when this bundle is fully installed
+	uint64 CurrentInstallSize = 0; // Disk footprint of the bundle in it's current state
 	bool bIsStartup = false; // Only one startup bundle allowed.  All sources must agree on this.
 	bool bDoPatchCheck = false; // This bundle should do a patch check and fail if it doesn't pass
 	bool bBundleUpToDate = false; // Whether this bundle is up to date
+	bool bIsCached = false;
 };
 
 struct FInstallBundleSourceBundleInfoQueryResultInfo
@@ -219,7 +223,7 @@ enum class EInstallBundleSourceUpdateBundleInfoResult : uint32
 	Count,
 };
 
-struct FInstallBundleSourceRequestResultInfo
+struct FInstallBundleSourceUpdateContentResultInfo
 {
 	FName BundleName;
 	EInstallBundleResult Result = EInstallBundleResult::OK;
@@ -239,6 +243,13 @@ struct FInstallBundleSourceRequestResultInfo
 	bool DidBundleSourceDoWork() const { return (ContentPaths.Num() != 0);} 
 };
 
+struct FInstallBundleSourceRemoveContentResultInfo
+{
+	FName BundleName;
+
+	bool bContentWasRemoved = false;
+};
+
 struct FInstallBundleSourceProgress
 {
 	FName BundleName;
