Configure editor builds to always use null install bundle manager implementation, regardless of platform.
Remove content patch check hack.

diff --git a/Engine/Config/BaseEngine.ini b/Engine/Config/BaseEngine.ini
--- a/Engine/Config/BaseEngine.ini
+++ b/Engine/Config/BaseEngine.ini
@@ -13,6 +13,7 @@ LocalMapOptions=
 
 [InstallBundleManager]
 ModuleName=NullInstallBundleManager
+EditorModuleName=NullInstallBundleManager
 
 [URL]
 Protocol=unreal
diff --git a/Engine/Source/Runtime/InstallBundleManager/Private/InstallBundleManagerInterface.cpp b/Engine/Source/Runtime/InstallBundleManager/Private/InstallBundleManagerInterface.cpp
--- a/Engine/Source/Runtime/InstallBundleManager/Private/InstallBundleManagerInterface.cpp
+++ b/Engine/Source/Runtime/InstallBundleManager/Private/InstallBundleManagerInterface.cpp
@@ -23,7 +23,11 @@ IInstallBundleManager* IInstallBundleManager::GetPlatformInstallBundleManager()
 	{
 		FString ModuleName;
 		IInstallBundleManagerModule* Module = nullptr;
+#if WITH_EDITOR
 		GConfig->GetString(TEXT("InstallBundleManager"), TEXT("ModuleName"), ModuleName, GEngineIni);
+#else
+		GConfig->GetString(TEXT("InstallBundleManager"), TEXT("EditorModuleName"), ModuleName, GEngineIni);
+#endif // WITH_EDITOR
 
 		if (FModuleManager::Get().ModuleExists(*ModuleName))
 		{
diff --git a/Engine/Source/Runtime/InstallBundleManager/Public/InstallBundleManagerModule.h b/Engine/Source/Runtime/InstallBundleManager/Public/InstallBundleManagerModule.h
--- a/Engine/Source/Runtime/InstallBundleManager/Public/InstallBundleManagerModule.h
+++ b/Engine/Source/Runtime/InstallBundleManager/Public/InstallBundleManagerModule.h
@@ -53,7 +53,11 @@ public:
 	{
 		// Only instantiate the bundle manager if this is the version the game has been configured to use
 		FString ModuleName;
+#if WITH_EDITOR
 		GConfig->GetString(TEXT("InstallBundleManager"), TEXT("ModuleName"), ModuleName, GEngineIni);
+#else
+		GConfig->GetString(TEXT("InstallBundleManager"), TEXT("EditorModuleName"), ModuleName, GEngineIni);
+#endif // WITH_EDITOR
 
 		if (FModuleManager::Get().GetModule(*ModuleName) == this)
 		{
