Fix my last CL

diff --git a/Engine/Source/Runtime/InstallBundleManager/Private/InstallBundleManagerInterface.cpp b/Engine/Source/Runtime/InstallBundleManager/Private/InstallBundleManagerInterface.cpp
--- a/Engine/Source/Runtime/InstallBundleManager/Private/InstallBundleManagerInterface.cpp
+++ b/Engine/Source/Runtime/InstallBundleManager/Private/InstallBundleManagerInterface.cpp
@@ -24,9 +24,9 @@ IInstallBundleManager* IInstallBundleManager::GetPlatformInstallBundleManager()
 		FString ModuleName;
 		IInstallBundleManagerModule* Module = nullptr;
 #if WITH_EDITOR
-		GConfig->GetString(TEXT("InstallBundleManager"), TEXT("ModuleName"), ModuleName, GEngineIni);
-#else
 		GConfig->GetString(TEXT("InstallBundleManager"), TEXT("EditorModuleName"), ModuleName, GEngineIni);
+#else
+		GConfig->GetString(TEXT("InstallBundleManager"), TEXT("ModuleName"), ModuleName, GEngineIni);
 #endif // WITH_EDITOR
 
 		if (FModuleManager::Get().ModuleExists(*ModuleName))
diff --git a/Engine/Source/Runtime/InstallBundleManager/Public/InstallBundleManagerModule.h b/Engine/Source/Runtime/InstallBundleManager/Public/InstallBundleManagerModule.h
--- a/Engine/Source/Runtime/InstallBundleManager/Public/InstallBundleManagerModule.h
+++ b/Engine/Source/Runtime/InstallBundleManager/Public/InstallBundleManagerModule.h
@@ -54,9 +54,9 @@ public:
 		// Only instantiate the bundle manager if this is the version the game has been configured to use
 		FString ModuleName;
 #if WITH_EDITOR
-		GConfig->GetString(TEXT("InstallBundleManager"), TEXT("ModuleName"), ModuleName, GEngineIni);
-#else
 		GConfig->GetString(TEXT("InstallBundleManager"), TEXT("EditorModuleName"), ModuleName, GEngineIni);
+#else
+		GConfig->GetString(TEXT("InstallBundleManager"), TEXT("ModuleName"), ModuleName, GEngineIni);
 #endif // WITH_EDITOR
 
 		if (FModuleManager::Get().GetModule(*ModuleName) == this)
