Removed UE4 reference from comments

diff --git a/Engine/Source/Runtime/InstallBundleManager/Public/InstallBundleUtils.h b/Engine/Source/Runtime/InstallBundleManager/Public/InstallBundleUtils.h
--- a/Engine/Source/Runtime/InstallBundleManager/Public/InstallBundleUtils.h
+++ b/Engine/Source/Runtime/InstallBundleManager/Public/InstallBundleUtils.h
@@ -330,7 +330,7 @@ namespace InstallBundleUtil
 			
 		//Static Methods
 		public:
-			//Helper function that generates a basic ExpectedAnalyticsID from device and UE4 information
+			//Helper function that generates a basic ExpectedAnalyticsID from device and UE information
 			static const FString GetBaseExpectedAnalyticsID();
 		
 		//Pure Virtual Methods
