Add EInstallBundleResult::ConnectivityError error code.

diff --git a/Engine/Source/Runtime/InstallBundleManager/Private/InstallBundleTypes.cpp b/Engine/Source/Runtime/InstallBundleManager/Private/InstallBundleTypes.cpp
--- a/Engine/Source/Runtime/InstallBundleManager/Private/InstallBundleTypes.cpp
+++ b/Engine/Source/Runtime/InstallBundleManager/Private/InstallBundleTypes.cpp
@@ -83,6 +83,7 @@ const TCHAR* LexToString(EInstallBundleResult Result)
 		TEXT("InstallError"),
 		TEXT("InstallerOutOfDiskSpaceError"),
 		TEXT("ManifestArchiveError"),
+		TEXT("ConnectivityError"),
 		TEXT("UserCancelledError"),
 		TEXT("InitializationError"),
 		TEXT("InitializationPending"),
diff --git a/Engine/Source/Runtime/InstallBundleManager/Public/InstallBundleTypes.h b/Engine/Source/Runtime/InstallBundleManager/Public/InstallBundleTypes.h
--- a/Engine/Source/Runtime/InstallBundleManager/Public/InstallBundleTypes.h
+++ b/Engine/Source/Runtime/InstallBundleManager/Public/InstallBundleTypes.h
@@ -121,6 +121,7 @@ enum class EInstallBundleResult : int
 	InstallError,
 	InstallerOutOfDiskSpaceError,
 	ManifestArchiveError,
+	ConnectivityError,
 	UserCancelledError,
 	InitializationError,
 	InitializationPending,
