Add an init complete delegate to bundle manager

diff --git a/Engine/Source/Runtime/InstallBundleManager/Private/InstallBundleManagerInterface.cpp b/Engine/Source/Runtime/InstallBundleManager/Private/InstallBundleManagerInterface.cpp
--- a/Engine/Source/Runtime/InstallBundleManager/Private/InstallBundleManagerInterface.cpp
+++ b/Engine/Source/Runtime/InstallBundleManager/Private/InstallBundleManagerInterface.cpp
@@ -3,6 +3,8 @@
 #include "InstallBundleManagerInterface.h"
 #include "InstallBundleManagerModule.h"
 
+FInstallBundleManagerInitCompleteMultiDelegate IInstallBundleManager::InitCompleteDelegate;
+
 FInstallBundleCompleteMultiDelegate IInstallBundleManager::InstallBundleCompleteDelegate;
 
 FInstallBundlePausedMultiDelegate IInstallBundleManager::PausedBundleDelegate;
diff --git a/Engine/Source/Runtime/InstallBundleManager/Public/InstallBundleManagerInterface.h b/Engine/Source/Runtime/InstallBundleManager/Public/InstallBundleManagerInterface.h
--- a/Engine/Source/Runtime/InstallBundleManager/Public/InstallBundleManagerInterface.h
+++ b/Engine/Source/Runtime/InstallBundleManager/Public/InstallBundleManagerInterface.h
@@ -59,6 +59,7 @@ enum class EInstallBundleManagerInitErrorHandlerResult
 using FInstallBundleSourceOrCache = TUnion<EInstallBundleSourceType, FName>;
 
 DECLARE_DELEGATE_RetVal_OneParam(EInstallBundleManagerInitErrorHandlerResult, FInstallBundleManagerInitErrorHandler, EInstallBundleManagerInitResult);
+DECLARE_MULTICAST_DELEGATE_OneParam(FInstallBundleManagerInitCompleteMultiDelegate, EInstallBundleManagerInitResult);
 
 DECLARE_MULTICAST_DELEGATE_OneParam(FInstallBundleCompleteMultiDelegate, FInstallBundleRequestResultInfo);
 DECLARE_MULTICAST_DELEGATE_OneParam(FInstallBundlePausedMultiDelegate, FInstallBundlePauseInfo);
@@ -75,6 +76,8 @@ DECLARE_DELEGATE(FInstallBundleManagerFlushCacheCompleteDelegate);
 class INSTALLBUNDLEMANAGER_API IInstallBundleManager : public TSharedFromThis<IInstallBundleManager>
 {
 public:
+	static FInstallBundleManagerInitCompleteMultiDelegate InitCompleteDelegate;
+
 	static FInstallBundleCompleteMultiDelegate InstallBundleCompleteDelegate; // Called when a content request is complete
 	static FInstallBundlePausedMultiDelegate PausedBundleDelegate;
 	static FInstallBundleReleasedMultiDelegate ReleasedDelegate; // Called when content release request is complete
