Add EInstallBundleSourceUpdateBundleInfoResult::IlegalInstallSizeChange

diff --git a/Engine/Source/Runtime/InstallBundleManager/Private/InstallBundleTypes.cpp b/Engine/Source/Runtime/InstallBundleManager/Private/InstallBundleTypes.cpp
--- a/Engine/Source/Runtime/InstallBundleManager/Private/InstallBundleTypes.cpp
+++ b/Engine/Source/Runtime/InstallBundleManager/Private/InstallBundleTypes.cpp
@@ -153,6 +153,7 @@ const TCHAR* LexToString(EInstallBundleSourceUpdateBundleInfoResult Result)
 		TEXT("AlreadyRequested"),
 		TEXT("IllegalStartupBundle"),
 		TEXT("IllegalCacheStatus"),
+		TEXT("IllegalInstallSizeChange"),
 	};
 
 	return InstallBundleUtil::TLexToString(Result, Strings);
diff --git a/Engine/Source/Runtime/InstallBundleManager/Public/InstallBundleTypes.h b/Engine/Source/Runtime/InstallBundleManager/Public/InstallBundleTypes.h
--- a/Engine/Source/Runtime/InstallBundleManager/Public/InstallBundleTypes.h
+++ b/Engine/Source/Runtime/InstallBundleManager/Public/InstallBundleTypes.h
@@ -252,6 +252,7 @@ enum class EInstallBundleSourceUpdateBundleInfoResult : uint8
 	AlreadyRequested,
 	IllegalStartupBundle,
 	IllegalCacheStatus,
+	IllegalInstallSizeChange,
 	Count,
 };
 INSTALLBUNDLEMANAGER_API const TCHAR* LexToString(EInstallBundleSourceUpdateBundleInfoResult Result);
