Prep for combined bundle cache - Add age scalar so we can have different cache eviction rates for different bundle sources.
Fix for crash if manifest is missing from an installed bundle - this mostly affects bundles that were installed from a <redacted> since they don't get manifests.
When pruning no longer relevant bundles, remove from any associated caches.
Add a console command to flush install bundle caches.

diff --git a/Engine/Source/Runtime/InstallBundleManager/Private/InstallBundleCache.cpp b/Engine/Source/Runtime/InstallBundleManager/Private/InstallBundleCache.cpp
--- a/Engine/Source/Runtime/InstallBundleManager/Private/InstallBundleCache.cpp
+++ b/Engine/Source/Runtime/InstallBundleManager/Private/InstallBundleCache.cpp
@@ -24,6 +24,7 @@ void FInstallBundleCache::AddOrUpdateBundle(EInstallBundleSourceType Source, con
 	Info.FullInstallSize = AddInfo.FullInstallSize;
 	Info.CurrentInstallSize = AddInfo.CurrentInstallSize;
 	Info.TimeStamp = AddInfo.TimeStamp;
+	Info.AgeScalar = FMath::Clamp(AddInfo.AgeScalar, 0.1, 1.0);
 
 	UpdateCacheInfoFromSourceInfo(AddInfo.BundleName);
 
@@ -62,6 +63,7 @@ TOptional<FInstallBundleCacheBundleInfo> FInstallBundleCache::GetBundleInfo(EIns
 			OutInfo.FullInstallSize = SourceInfo->FullInstallSize;
 			OutInfo.CurrentInstallSize = SourceInfo->CurrentInstallSize;
 			OutInfo.TimeStamp = SourceInfo->TimeStamp;
+			OutInfo.AgeScalar = SourceInfo->AgeScalar;
 		}
 	}
 
@@ -148,11 +150,14 @@ FInstallBundleCacheReserveResult FInstallBundleCache::Reserve(FName BundleName)
 
 	// TODO: Bundles that have BundleSize > 0 or are PendingEvict should be 
 	// sorted to the beginning.  We should be able to stop iterating sooner in that case.
-	CacheInfo.ValueSort([](const FBundleCacheInfo& A, const FBundleCacheInfo& B)
+	CacheInfo.ValueSort([Now=FDateTime::UtcNow()](const FBundleCacheInfo& A, const FBundleCacheInfo& B)
 	{
 		if (A.IsHintRequested() == B.IsHintRequested())
 		{
-			return A.TimeStamp < B.TimeStamp;
+			FTimespan AgeA = (Now > A.TimeStamp) ? Now - A.TimeStamp : FTimespan(0);
+			FTimespan AgeB = (Now > B.TimeStamp) ? Now - B.TimeStamp : FTimespan(0);
+
+			return AgeA * A.AgeScalar > AgeB * B.AgeScalar;
 		}
 		
 		return !A.IsHintRequested() && B.IsHintRequested();
@@ -424,15 +429,24 @@ void FInstallBundleCache::UpdateCacheInfoFromSourceInfo(FName BundleName)
 	}
 
 	FDateTime TimeStamp = FDateTime::MinValue();
+	double AgeScalar = 1.0;
 	uint64 FullInstallSize = 0;
 	uint64 CurrentInstallSize = 0;
 	for (const TPair<EInstallBundleSourceType, FPerSourceBundleCacheInfo>& Pair : *SourcesMap)
 	{
 		FullInstallSize += Pair.Value.FullInstallSize;
 		CurrentInstallSize += Pair.Value.CurrentInstallSize;
-		if (Pair.Value.CurrentInstallSize > 0 && Pair.Value.TimeStamp > TimeStamp)
+		if (Pair.Value.CurrentInstallSize > 0)
 		{
-			TimeStamp = Pair.Value.TimeStamp;
+			if (Pair.Value.TimeStamp > TimeStamp)
+			{
+				TimeStamp = Pair.Value.TimeStamp;
+			}
+
+			if (Pair.Value.AgeScalar < AgeScalar)
+			{
+				AgeScalar = Pair.Value.AgeScalar;
+			}
 		}
 	}
 
@@ -442,4 +456,5 @@ void FInstallBundleCache::UpdateCacheInfoFromSourceInfo(FName BundleName)
 	BundleCacheInfo.FullInstallSize = FullInstallSize;
 	BundleCacheInfo.CurrentInstallSize = CurrentInstallSize;
 	BundleCacheInfo.TimeStamp = TimeStamp;
+	BundleCacheInfo.AgeScalar = AgeScalar;
 }
diff --git a/Engine/Source/Runtime/InstallBundleManager/Public/InstallBundleCache.h b/Engine/Source/Runtime/InstallBundleManager/Public/InstallBundleCache.h
--- a/Engine/Source/Runtime/InstallBundleManager/Public/InstallBundleCache.h
+++ b/Engine/Source/Runtime/InstallBundleManager/Public/InstallBundleCache.h
@@ -18,6 +18,7 @@ struct FInstallBundleCacheBundleInfo
 	uint64 FullInstallSize = 0; // Total disk footprint when this bundle is fully installed
 	uint64 CurrentInstallSize = 0; // Disk footprint of the bundle in it's current state
 	FDateTime TimeStamp = FDateTime::MinValue(); // Last access time for the bundle.  Used for eviction order
+	double AgeScalar = 1.0; // Allow some bundles to "age" slower than others
 };
 
 enum class EInstallBundleCacheReserveResult : int8
@@ -94,6 +95,7 @@ private:
 		uint64 FullInstallSize = 0;
 		uint64 CurrentInstallSize = 0;
 		FDateTime TimeStamp = FDateTime::MinValue();
+		double AgeScalar = 1.0;
 	};
 
 	enum class ECacheState : uint8
@@ -108,6 +110,7 @@ private:
 		uint64 FullInstallSize = 0;
 		uint64 CurrentInstallSize = 0;
 		FDateTime TimeStamp = FDateTime::MinValue();
+		double AgeScalar = 1.0;
 		ECacheState State = ECacheState::Released;
 		int32 HintReqeustedCount = 0; // Hint to the cache that this bundle is requested, and we should prefer to evict non-requested bundles if possible
 
diff --git a/Engine/Source/Runtime/InstallBundleManager/Public/InstallBundleSourceInterface.h b/Engine/Source/Runtime/InstallBundleManager/Public/InstallBundleSourceInterface.h
--- a/Engine/Source/Runtime/InstallBundleManager/Public/InstallBundleSourceInterface.h
+++ b/Engine/Source/Runtime/InstallBundleManager/Public/InstallBundleSourceInterface.h
@@ -31,6 +31,9 @@ public:
 	// Returns the how this source should be weighted when combined with other sources
 	virtual float GetSourceWeight() const { return 1.0f; }
 
+	// Returns a value in the [0.1, 1.0] range that allows bundles using this source to age out of the LRU cache more slowly
+	virtual double GetSourceCacheAgeScalar() const { return 1.0; }
+
 	// Called once by bundle manager after constructing the bundle source
 	// Any non-fallback errors returned will cause bundle manager to fail to initialize
 	virtual FInstallBundleSourceInitInfo Init(
