Improve GetCacheStats
 - Now dumps bundle info in eviction order.
 - Add verbosity control
 - Add version that dumps stats for all caches

diff --git a/Engine/Source/Runtime/InstallBundleManager/Private/InstallBundleCache.cpp b/Engine/Source/Runtime/InstallBundleManager/Private/InstallBundleCache.cpp
--- a/Engine/Source/Runtime/InstallBundleManager/Private/InstallBundleCache.cpp
+++ b/Engine/Source/Runtime/InstallBundleManager/Private/InstallBundleCache.cpp
@@ -150,18 +150,7 @@ FInstallBundleCacheReserveResult FInstallBundleCache::Reserve(FName BundleName)
 
 	// TODO: Bundles that have BundleSize > 0 or are PendingEvict should be 
 	// sorted to the beginning.  We should be able to stop iterating sooner in that case.
-	CacheInfo.ValueSort([Now=FDateTime::UtcNow()](const FBundleCacheInfo& A, const FBundleCacheInfo& B)
-	{
-		if (A.IsHintRequested() == B.IsHintRequested())
-		{
-			FTimespan AgeA = (Now > A.TimeStamp) ? Now - A.TimeStamp : FTimespan(0);
-			FTimespan AgeB = (Now > B.TimeStamp) ? Now - B.TimeStamp : FTimespan(0);
-
-			return AgeA * A.AgeScalar > AgeB * B.AgeScalar;
-		}
-		
-		return !A.IsHintRequested() && B.IsHintRequested();
-	});
+	CacheInfo.ValueSort(FCacheSortPredicate());
 
 	uint64 CanFreeSpace = 0;
 	for (const TPair<FName, FBundleCacheInfo>& Pair : CacheInfo)
@@ -363,7 +352,7 @@ void FInstallBundleCache::CheckInvariants() const
 #endif // INSTALLBUNDLE_CACHE_CHECK_INVARIANTS
 }
 
-FInstallBundleCacheStats FInstallBundleCache::GetStats(bool bDumpToLog /*= false*/) const
+FInstallBundleCacheStats FInstallBundleCache::GetStats(bool bDumpToLog /*= false*/, bool bVerbose /*= false*/) const
 {
 	FInstallBundleCacheStats Stats;
 	Stats.CacheName = CacheName;
@@ -373,6 +362,9 @@ FInstallBundleCacheStats FInstallBundleCache::GetStats(bool bDumpToLog /*= false
 	{
 		UE_LOG(LogInstallBundleManager, Display, TEXT("\n"));
 		UE_LOG(LogInstallBundleManager, Display, TEXT("*Install Bundle Cache Stats %s"), *CacheName.ToString());
+
+		// Dump info in eviction order
+		CacheInfo.ValueSort(FCacheSortPredicate());
 	}
 
 	for (const TPair<FName, FBundleCacheInfo>& CachePair : CacheInfo)
@@ -388,11 +380,24 @@ FInstallBundleCacheStats FInstallBundleCache::GetStats(bool bDumpToLog /*= false
 
 		if (bDumpToLog && (Info.CurrentInstallSize > 0 || Info.State != ECacheState::Released))
 		{
-			UE_LOG(LogInstallBundleManager, Verbose, TEXT("*\tbundle %s"), *CachePair.Key.ToString());
-			UE_LOG(LogInstallBundleManager, Verbose, TEXT("*\t\tfull size: %" UINT64_FMT), Info.FullInstallSize);
-			UE_LOG(LogInstallBundleManager, Verbose, TEXT("*\t\tcurrent size: %" UINT64_FMT), Info.CurrentInstallSize);
-			UE_LOG(LogInstallBundleManager, Verbose, TEXT("*\t\treserved: %s"), (Info.State == ECacheState::Reserved) ? TEXT("true") : TEXT("false"));
-			UE_LOG(LogInstallBundleManager, Verbose, TEXT("*\t\ttimestamp: %s"), *Info.TimeStamp.ToString());
+			if (bVerbose)
+			{
+				UE_LOG(LogInstallBundleManager, Display, TEXT("*\tbundle %s"), *CachePair.Key.ToString());
+				UE_LOG(LogInstallBundleManager, Display, TEXT("*\t\tfull size: %" UINT64_FMT), Info.FullInstallSize);
+				UE_LOG(LogInstallBundleManager, Display, TEXT("*\t\tcurrent size: %" UINT64_FMT), Info.CurrentInstallSize);
+				UE_LOG(LogInstallBundleManager, Display, TEXT("*\t\treserved: %s"), (Info.State == ECacheState::Reserved) ? TEXT("true") : TEXT("false"));
+				UE_LOG(LogInstallBundleManager, Display, TEXT("*\t\ttimestamp: %s"), *Info.TimeStamp.ToString());
+				UE_LOG(LogInstallBundleManager, Display, TEXT("*\t\tage scale: %f"), Info.AgeScalar);
+			}
+			else
+			{
+				UE_LOG(LogInstallBundleManager, Verbose, TEXT("*\tbundle %s"), *CachePair.Key.ToString());
+				UE_LOG(LogInstallBundleManager, Verbose, TEXT("*\t\tfull size: %" UINT64_FMT), Info.FullInstallSize);
+				UE_LOG(LogInstallBundleManager, Verbose, TEXT("*\t\tcurrent size: %" UINT64_FMT), Info.CurrentInstallSize);
+				UE_LOG(LogInstallBundleManager, Verbose, TEXT("*\t\treserved: %s"), (Info.State == ECacheState::Reserved) ? TEXT("true") : TEXT("false"));
+				UE_LOG(LogInstallBundleManager, Verbose, TEXT("*\t\ttimestamp: %s"), *Info.TimeStamp.ToString());
+				UE_LOG(LogInstallBundleManager, Verbose, TEXT("*\t\tage scale: %f"), Info.AgeScalar);
+			}
 		}
 	}
 
diff --git a/Engine/Source/Runtime/InstallBundleManager/Public/InstallBundleCache.h b/Engine/Source/Runtime/InstallBundleManager/Public/InstallBundleCache.h
--- a/Engine/Source/Runtime/InstallBundleManager/Public/InstallBundleCache.h
+++ b/Engine/Source/Runtime/InstallBundleManager/Public/InstallBundleCache.h
@@ -80,7 +80,7 @@ public:
 	// Hint to the cache that this bundle is requested, and we should prefer to evict non-requested bundles if possible
 	void HintRequested(FName BundleName, bool bRequested);
 
-	FInstallBundleCacheStats GetStats(bool bDumpToLog = false) const;
+	FInstallBundleCacheStats GetStats(bool bDumpToLog = false, bool bVerbose = false) const;
 
 private:
 	uint64 GetFreeSpaceInternal(uint64 UsedSize) const;
@@ -133,11 +133,31 @@ private:
 		}
 	};
 
+	struct FCacheSortPredicate
+	{
+		bool operator()(const FBundleCacheInfo& A, const FBundleCacheInfo& B) const
+		{
+			if (A.IsHintRequested() == B.IsHintRequested())
+			{
+				FTimespan AgeA = (Now > A.TimeStamp) ? Now - A.TimeStamp : FTimespan(0);
+				FTimespan AgeB = (Now > B.TimeStamp) ? Now - B.TimeStamp : FTimespan(0);
+
+				return AgeA * A.AgeScalar > AgeB * B.AgeScalar;
+			}
+
+			return !A.IsHintRequested() && B.IsHintRequested();
+		};
+
+	private:
+		FDateTime Now = FDateTime::UtcNow();
+	};
+
 private:
 
 	TMap<FName, TMap<EInstallBundleSourceType, FPerSourceBundleCacheInfo>> PerSourceCacheInfo;
 
-	TMap<FName, FBundleCacheInfo> CacheInfo;
+	// mutable to allow sorting in const contexts
+	mutable TMap<FName, FBundleCacheInfo> CacheInfo;
 
 	uint64 TotalSize = 0;
 
diff --git a/Engine/Source/Runtime/InstallBundleManager/Public/InstallBundleManagerInterface.h b/Engine/Source/Runtime/InstallBundleManager/Public/InstallBundleManagerInterface.h
--- a/Engine/Source/Runtime/InstallBundleManager/Public/InstallBundleManagerInterface.h
+++ b/Engine/Source/Runtime/InstallBundleManager/Public/InstallBundleManagerInterface.h
@@ -119,7 +119,8 @@ public:
 	EInstallBundleResult FlushCache(FInstallBundleManagerFlushCacheCompleteDelegate Callback, ELogVerbosity::Type LogVerbosityOverride = ELogVerbosity::NoLogging);
 	virtual EInstallBundleResult FlushCache(FInstallBundleSourceOrCache SourceOrCache, FInstallBundleManagerFlushCacheCompleteDelegate Callback, ELogVerbosity::Type LogVerbosityOverride = ELogVerbosity::NoLogging) = 0;
 
-	virtual TOptional<FInstallBundleCacheStats> GetCacheStats(FInstallBundleSourceOrCache SourceOrCache, bool bDumpToLog = false) = 0;
+	virtual TArray<FInstallBundleCacheStats> GetCacheStats(bool bDumpToLog = false, ELogVerbosity::Type LogVerbosityOverride = ELogVerbosity::NoLogging) = 0;
+	virtual TOptional<FInstallBundleCacheStats> GetCacheStats(FInstallBundleSourceOrCache SourceOrCache, bool bDumpToLog = false, ELogVerbosity::Type LogVerbosityOverride = ELogVerbosity::NoLogging) = 0;
 
 	void RequestRemoveContentOnNextInit(FName RemoveName, TArrayView<const FName> KeepNames = TArrayView<const FName>());
 	virtual void RequestRemoveContentOnNextInit(TArrayView<const FName> RemoveNames, TArrayView<const FName> KeepNames = TArrayView<const FName>()) = 0;
diff --git a/Engine/Source/Runtime/NullInstallBundleManager/Source/NullInstallBundleManager.cpp b/Engine/Source/Runtime/NullInstallBundleManager/Source/NullInstallBundleManager.cpp
--- a/Engine/Source/Runtime/NullInstallBundleManager/Source/NullInstallBundleManager.cpp
+++ b/Engine/Source/Runtime/NullInstallBundleManager/Source/NullInstallBundleManager.cpp
@@ -84,7 +84,12 @@ class FNullInstallBundleManager : public IInstallBundleManager
 		return EInstallBundleResult::OK;
 	}
 
-	virtual TOptional<FInstallBundleCacheStats> GetCacheStats(FInstallBundleSourceOrCache SourceOrCache, bool bDumpToLog = false) override
+	virtual TArray<FInstallBundleCacheStats> GetCacheStats(bool bDumpToLog = false, ELogVerbosity::Type LogVerbosityOverride = ELogVerbosity::NoLogging) override
+	{
+		return TArray<FInstallBundleCacheStats>();
+	}
+
+	virtual TOptional<FInstallBundleCacheStats> GetCacheStats(FInstallBundleSourceOrCache SourceOrCache, bool bDumpToLog = false, ELogVerbosity::Type LogVerbosityOverride = ELogVerbosity::NoLogging) override
 	{
 		return TOptional<FInstallBundleCacheStats>();
 	}
