Install Bundle Cache
- Add CSV format option for GetStats()

diff --git a/Engine/Source/Runtime/InstallBundleManager/Private/InstallBundleCache.cpp b/Engine/Source/Runtime/InstallBundleManager/Private/InstallBundleCache.cpp
--- a/Engine/Source/Runtime/InstallBundleManager/Private/InstallBundleCache.cpp
+++ b/Engine/Source/Runtime/InstallBundleManager/Private/InstallBundleCache.cpp
@@ -352,64 +352,146 @@ void FInstallBundleCache::CheckInvariants() const
 #endif // INSTALLBUNDLE_CACHE_CHECK_INVARIANTS
 }
 
-FInstallBundleCacheStats FInstallBundleCache::GetStats(bool bDumpToLog /*= false*/, bool bVerbose /*= false*/) const
+FInstallBundleCacheStats FInstallBundleCache::GetStats(EInstallBundleCacheDumpToLog DumpToLog /*= EInstallBundleCacheDumpToLog::None*/, bool bVerbose /*= false*/) const
 {
 	FInstallBundleCacheStats Stats;
 	Stats.CacheName = CacheName;
 	Stats.MaxSize = TotalSize;
 
+	bool bDumpToLog = DumpToLog != EInstallBundleCacheDumpToLog::None;
+
 	if (bDumpToLog)
 	{
-		UE_LOG(LogInstallBundleManager, Display, TEXT("\n"));
-		UE_LOG(LogInstallBundleManager, Display, TEXT("*Install Bundle Cache Stats %s"), *CacheName.ToString());
+		UE_LOG(LogInstallBundleManager, Display, TEXT(""));
+		UE_LOG(LogInstallBundleManager, Display, TEXT("* Install Bundle Cache Stats %s"), *CacheName.ToString());
 
 		// Dump info in eviction order
 		CacheInfo.ValueSort(FCacheSortPredicate());
 	}
 
-	for (const TPair<FName, FBundleCacheInfo>& CachePair : CacheInfo)
+	auto AddStats = [&Stats](const FBundleCacheInfo& Info)
 	{
-		const FBundleCacheInfo& Info = CachePair.Value;
-
 		Stats.UsedSize += Info.GetSize();
 
 		if (Info.State == ECacheState::Reserved)
 		{
 			Stats.ReservedSize += Info.CurrentInstallSize;
 		}
+	};
+
+	// Note: Yes this verbosity is correct and very non-intuitive
+
+#define INSTALLBUNDLECACHE_DEFAULT_LOG(Verbosity) \
+	if (Info.CurrentInstallSize > 0 || Info.State != ECacheState::Released) \
+	{ \
+		UE_LOG(LogInstallBundleManager, Verbosity, TEXT("* \tbundle %s"), *BundleName.ToString()) \
+		UE_LOG(LogInstallBundleManager, Verbosity, TEXT("* \t\tfull size: %" UINT64_FMT), Info.FullInstallSize) \
+		UE_LOG(LogInstallBundleManager, Verbosity, TEXT("* \t\tcurrent size: %" UINT64_FMT), Info.CurrentInstallSize) \
+		UE_LOG(LogInstallBundleManager, Verbosity, TEXT("* \t\treserved: %s"), (Info.State == ECacheState::Reserved) ? TEXT("true") : TEXT("false")) \
+		UE_LOG(LogInstallBundleManager, Verbosity, TEXT("* \t\ttimestamp: %s"), *Info.TimeStamp.ToString()) \
+		UE_LOG(LogInstallBundleManager, Verbosity, TEXT("* \t\tage scale: %f"), Info.AgeScalar) \
+	}
+
+#define INSTALLBUNDLECACHE_CSV_HEADER_LOG(Verbosity) \
+	UE_LOG(LogInstallBundleManager, Verbosity, TEXT("* \tbundle, full size, current size, reserved, timestamp, age scale"))
+
+#define INSTALLBUNDLECACHE_CSV_LOG(Verbosity) \
+	if (Info.CurrentInstallSize > 0 || Info.State != ECacheState::Released) \
+	{ \
+		UE_LOG(LogInstallBundleManager, Verbosity, TEXT("* \t%s, %" UINT64_FMT ", %" UINT64_FMT ", %s, %s, %f"), *BundleName.ToString(), Info.FullInstallSize, Info.CurrentInstallSize, (Info.State == ECacheState::Reserved) ? TEXT("true") : TEXT("false"), *Info.TimeStamp.ToString(), Info.AgeScalar) \
+	}
+
+	auto DumpToLog_Default = [](FName BundleName, const FBundleCacheInfo& Info)
+	{
+		INSTALLBUNDLECACHE_DEFAULT_LOG(Verbose);
+	};
+
+	auto DumpToLog_DefaultVerbose = [](FName BundleName, const FBundleCacheInfo& Info)
+	{
+		INSTALLBUNDLECACHE_DEFAULT_LOG(Display);
+	};
+
+	auto DumpToLog_CSV = [](FName BundleName, const FBundleCacheInfo& Info)
+	{
+		INSTALLBUNDLECACHE_CSV_LOG(Verbose);
+	};
+
+	auto DumpToLog_CSVVerbose = [](FName BundleName, const FBundleCacheInfo& Info)
+	{
+		INSTALLBUNDLECACHE_CSV_LOG(Display);
+	};
+
+	switch(DumpToLog)
+	{
+	default:
+	case EInstallBundleCacheDumpToLog::None:
+		for (const TPair<FName, FBundleCacheInfo>& CachePair : CacheInfo)
+		{
+			const FBundleCacheInfo& Info = CachePair.Value;
+			AddStats(Info);
+		}
+		break;
+		
+	case EInstallBundleCacheDumpToLog::Default:
+		if (bVerbose)
+		{
+			for (const TPair<FName, FBundleCacheInfo>& CachePair : CacheInfo)
+			{
+				const FBundleCacheInfo& Info = CachePair.Value;
+				AddStats(Info);
+				DumpToLog_DefaultVerbose(CachePair.Key, Info);
+			}
+		}
+		else
+		{
+			for (const TPair<FName, FBundleCacheInfo>& CachePair : CacheInfo)
+			{
+				const FBundleCacheInfo& Info = CachePair.Value;
+				AddStats(Info);
+				DumpToLog_Default(CachePair.Key, Info);
+			}
+		}
+		break;
 
-		if (bDumpToLog && (Info.CurrentInstallSize > 0 || Info.State != ECacheState::Released))
+	case EInstallBundleCacheDumpToLog::CSV:
+		if (bVerbose)
 		{
-			if (bVerbose)
+			INSTALLBUNDLECACHE_CSV_HEADER_LOG(Display);
+
+			for (const TPair<FName, FBundleCacheInfo>& CachePair : CacheInfo)
 			{
-				UE_LOG(LogInstallBundleManager, Display, TEXT("*\tbundle %s"), *CachePair.Key.ToString());
-				UE_LOG(LogInstallBundleManager, Display, TEXT("*\t\tfull size: %" UINT64_FMT), Info.FullInstallSize);
-				UE_LOG(LogInstallBundleManager, Display, TEXT("*\t\tcurrent size: %" UINT64_FMT), Info.CurrentInstallSize);
-				UE_LOG(LogInstallBundleManager, Display, TEXT("*\t\treserved: %s"), (Info.State == ECacheState::Reserved) ? TEXT("true") : TEXT("false"));
-				UE_LOG(LogInstallBundleManager, Display, TEXT("*\t\ttimestamp: %s"), *Info.TimeStamp.ToString());
-				UE_LOG(LogInstallBundleManager, Display, TEXT("*\t\tage scale: %f"), Info.AgeScalar);
+				const FBundleCacheInfo& Info = CachePair.Value;
+				AddStats(Info);
+				DumpToLog_CSVVerbose(CachePair.Key, Info);
 			}
-			else
+		}
+		else
+		{
+			INSTALLBUNDLECACHE_CSV_HEADER_LOG(Verbose);
+
+			for (const TPair<FName, FBundleCacheInfo>& CachePair : CacheInfo)
 			{
-				UE_LOG(LogInstallBundleManager, Verbose, TEXT("*\tbundle %s"), *CachePair.Key.ToString());
-				UE_LOG(LogInstallBundleManager, Verbose, TEXT("*\t\tfull size: %" UINT64_FMT), Info.FullInstallSize);
-				UE_LOG(LogInstallBundleManager, Verbose, TEXT("*\t\tcurrent size: %" UINT64_FMT), Info.CurrentInstallSize);
-				UE_LOG(LogInstallBundleManager, Verbose, TEXT("*\t\treserved: %s"), (Info.State == ECacheState::Reserved) ? TEXT("true") : TEXT("false"));
-				UE_LOG(LogInstallBundleManager, Verbose, TEXT("*\t\ttimestamp: %s"), *Info.TimeStamp.ToString());
-				UE_LOG(LogInstallBundleManager, Verbose, TEXT("*\t\tage scale: %f"), Info.AgeScalar);
+				const FBundleCacheInfo& Info = CachePair.Value;
+				AddStats(Info);
+				DumpToLog_CSV(CachePair.Key, Info);
 			}
 		}
+		break;
 	}
 
+#undef INSTALLBUNDLECACHE_CSV_LOG
+#undef INSTALLBUNDLECACHE_CSV_HEADER_LOG
+#undef INSTALLBUNDLECACHE_DEFAULT_LOG
+
 	Stats.FreeSize = GetFreeSpaceInternal(Stats.UsedSize);
 
 	if (bDumpToLog)
 	{
-		UE_LOG(LogInstallBundleManager, Display, TEXT("*\tsize: %" UINT64_FMT), Stats.MaxSize);
-		UE_LOG(LogInstallBundleManager, Display, TEXT("*\tused: %" UINT64_FMT), Stats.UsedSize);
-		UE_LOG(LogInstallBundleManager, Display, TEXT("*\treserved: %" UINT64_FMT), Stats.ReservedSize);
-		UE_LOG(LogInstallBundleManager, Display, TEXT("*\tfree: %" UINT64_FMT), Stats.FreeSize);
-		UE_LOG(LogInstallBundleManager, Display, TEXT("\n"));
+		UE_LOG(LogInstallBundleManager, Display, TEXT("* \tsize: %" UINT64_FMT), Stats.MaxSize);
+		UE_LOG(LogInstallBundleManager, Display, TEXT("* \tused: %" UINT64_FMT), Stats.UsedSize);
+		UE_LOG(LogInstallBundleManager, Display, TEXT("* \treserved: %" UINT64_FMT), Stats.ReservedSize);
+		UE_LOG(LogInstallBundleManager, Display, TEXT("* \tfree: %" UINT64_FMT), Stats.FreeSize);
+		UE_LOG(LogInstallBundleManager, Display, TEXT(""));
 	}
 
 	return Stats;
diff --git a/Engine/Source/Runtime/InstallBundleManager/Public/InstallBundleCache.h b/Engine/Source/Runtime/InstallBundleManager/Public/InstallBundleCache.h
--- a/Engine/Source/Runtime/InstallBundleManager/Public/InstallBundleCache.h
+++ b/Engine/Source/Runtime/InstallBundleManager/Public/InstallBundleCache.h
@@ -80,7 +80,7 @@ public:
 	// Hint to the cache that this bundle is requested, and we should prefer to evict non-requested bundles if possible
 	void HintRequested(FName BundleName, bool bRequested);
 
-	FInstallBundleCacheStats GetStats(bool bDumpToLog = false, bool bVerbose = false) const;
+	FInstallBundleCacheStats GetStats(EInstallBundleCacheDumpToLog DumpToLog = EInstallBundleCacheDumpToLog::None, bool bVerbose = false) const;
 
 private:
 	uint64 GetFreeSpaceInternal(uint64 UsedSize) const;
diff --git a/Engine/Source/Runtime/InstallBundleManager/Public/InstallBundleManagerInterface.h b/Engine/Source/Runtime/InstallBundleManager/Public/InstallBundleManagerInterface.h
--- a/Engine/Source/Runtime/InstallBundleManager/Public/InstallBundleManagerInterface.h
+++ b/Engine/Source/Runtime/InstallBundleManager/Public/InstallBundleManagerInterface.h
@@ -119,8 +119,8 @@ public:
 	EInstallBundleResult FlushCache(FInstallBundleManagerFlushCacheCompleteDelegate Callback, ELogVerbosity::Type LogVerbosityOverride = ELogVerbosity::NoLogging);
 	virtual EInstallBundleResult FlushCache(FInstallBundleSourceOrCache SourceOrCache, FInstallBundleManagerFlushCacheCompleteDelegate Callback, ELogVerbosity::Type LogVerbosityOverride = ELogVerbosity::NoLogging) = 0;
 
-	virtual TArray<FInstallBundleCacheStats> GetCacheStats(bool bDumpToLog = false, ELogVerbosity::Type LogVerbosityOverride = ELogVerbosity::NoLogging) = 0;
-	virtual TOptional<FInstallBundleCacheStats> GetCacheStats(FInstallBundleSourceOrCache SourceOrCache, bool bDumpToLog = false, ELogVerbosity::Type LogVerbosityOverride = ELogVerbosity::NoLogging) = 0;
+	virtual TArray<FInstallBundleCacheStats> GetCacheStats(EInstallBundleCacheDumpToLog DumpToLog = EInstallBundleCacheDumpToLog::None, ELogVerbosity::Type LogVerbosityOverride = ELogVerbosity::NoLogging) = 0;
+	virtual TOptional<FInstallBundleCacheStats> GetCacheStats(FInstallBundleSourceOrCache SourceOrCache, EInstallBundleCacheDumpToLog DumpToLog = EInstallBundleCacheDumpToLog::None, ELogVerbosity::Type LogVerbosityOverride = ELogVerbosity::NoLogging) = 0;
 
 	void RequestRemoveContentOnNextInit(FName RemoveName, TArrayView<const FName> KeepNames = TArrayView<const FName>());
 	virtual void RequestRemoveContentOnNextInit(TArrayView<const FName> RemoveNames, TArrayView<const FName> KeepNames = TArrayView<const FName>()) = 0;
diff --git a/Engine/Source/Runtime/InstallBundleManager/Public/InstallBundleTypes.h b/Engine/Source/Runtime/InstallBundleManager/Public/InstallBundleTypes.h
--- a/Engine/Source/Runtime/InstallBundleManager/Public/InstallBundleTypes.h
+++ b/Engine/Source/Runtime/InstallBundleManager/Public/InstallBundleTypes.h
@@ -335,3 +335,10 @@ struct FInstallBundleCacheStats
 	uint64 ReservedSize = 0;
 	uint64 FreeSize = 0;
 };
+
+enum class EInstallBundleCacheDumpToLog : int8
+{
+	None = 0,
+	Default,
+	CSV
+};
diff --git a/Engine/Source/Runtime/NullInstallBundleManager/Source/NullInstallBundleManager.cpp b/Engine/Source/Runtime/NullInstallBundleManager/Source/NullInstallBundleManager.cpp
--- a/Engine/Source/Runtime/NullInstallBundleManager/Source/NullInstallBundleManager.cpp
+++ b/Engine/Source/Runtime/NullInstallBundleManager/Source/NullInstallBundleManager.cpp
@@ -84,12 +84,12 @@ class FNullInstallBundleManager : public IInstallBundleManager
 		return EInstallBundleResult::OK;
 	}
 
-	virtual TArray<FInstallBundleCacheStats> GetCacheStats(bool bDumpToLog = false, ELogVerbosity::Type LogVerbosityOverride = ELogVerbosity::NoLogging) override
+	virtual TArray<FInstallBundleCacheStats> GetCacheStats(EInstallBundleCacheDumpToLog DumpToLog = EInstallBundleCacheDumpToLog::None, ELogVerbosity::Type LogVerbosityOverride = ELogVerbosity::NoLogging) override
 	{
 		return TArray<FInstallBundleCacheStats>();
 	}
 
-	virtual TOptional<FInstallBundleCacheStats> GetCacheStats(FInstallBundleSourceOrCache SourceOrCache, bool bDumpToLog = false, ELogVerbosity::Type LogVerbosityOverride = ELogVerbosity::NoLogging) override
+	virtual TOptional<FInstallBundleCacheStats> GetCacheStats(FInstallBundleSourceOrCache SourceOrCache, EInstallBundleCacheDumpToLog DumpToLog = EInstallBundleCacheDumpToLog::None, ELogVerbosity::Type LogVerbosityOverride = ELogVerbosity::NoLogging) override
 	{
 		return TOptional<FInstallBundleCacheStats>();
 	}
