Add Contains() accessor.

diff --git a/Engine/Source/Runtime/InstallBundleManager/Private/InstallBundleCache.cpp b/Engine/Source/Runtime/InstallBundleManager/Private/InstallBundleCache.cpp
--- a/Engine/Source/Runtime/InstallBundleManager/Private/InstallBundleCache.cpp
+++ b/Engine/Source/Runtime/InstallBundleManager/Private/InstallBundleCache.cpp
@@ -235,6 +235,22 @@ FInstallBundleCacheFlushResult FInstallBundleCache::Flush(EInstallBundleSourceTy
 	return Result;
 }
 
+bool FInstallBundleCache::Contains(FName BundleName) const
+{
+	return CacheInfo.Contains(BundleName);
+}
+
+bool FInstallBundleCache::Contains(EInstallBundleSourceType Source, FName BundleName) const
+{
+	const TMap<EInstallBundleSourceType, FPerSourceBundleCacheInfo>* SourcesMap = PerSourceCacheInfo.Find(BundleName);
+	if (SourcesMap)
+	{
+		return SourcesMap->Contains(Source);
+	}
+
+	return false;
+}
+
 bool FInstallBundleCache::Release(FName BundleName)
 {
 	CSV_SCOPED_TIMING_STAT(InstallBundleManager, FInstallBundleCache_Release);
diff --git a/Engine/Source/Runtime/InstallBundleManager/Public/InstallBundleCache.h b/Engine/Source/Runtime/InstallBundleManager/Public/InstallBundleCache.h
--- a/Engine/Source/Runtime/InstallBundleManager/Public/InstallBundleCache.h
+++ b/Engine/Source/Runtime/InstallBundleManager/Public/InstallBundleCache.h
@@ -47,7 +47,7 @@ public:
 
 	void Init(FInstallBundleCacheInitInfo InitInfo);
 
-	FName GetName() { return CacheName; }
+	FName GetName() const { return CacheName; }
 
 	// Add a bundle to the cache.  
 	void AddOrUpdateBundle(EInstallBundleSourceType Source, const FInstallBundleCacheBundleInfo& AddInfo);
@@ -70,6 +70,9 @@ public:
 	// Called from bundle manager, returns all bundles that can be evicted
 	FInstallBundleCacheFlushResult Flush(EInstallBundleSourceType* Source = nullptr);
 
+	bool Contains(FName BundleName) const;
+	bool Contains(EInstallBundleSourceType Source, FName BundleName) const;
+
 	// Called from bundle manager to make the files for this bundle eligible for eviction
 	bool Release(FName BundleName);
 
