Add a diff field when dumping installbundle cache using CSV format.

diff --git a/Engine/Source/Runtime/InstallBundleManager/Private/InstallBundleCache.cpp b/Engine/Source/Runtime/InstallBundleManager/Private/InstallBundleCache.cpp
--- a/Engine/Source/Runtime/InstallBundleManager/Private/InstallBundleCache.cpp
+++ b/Engine/Source/Runtime/InstallBundleManager/Private/InstallBundleCache.cpp
@@ -409,12 +409,15 @@ FInstallBundleCacheStats FInstallBundleCache::GetStats(EInstallBundleCacheDumpTo
 	}
 
 #define INSTALLBUNDLECACHE_CSV_HEADER_LOG(Verbosity) \
-	UE_LOG(LogInstallBundleManager, Verbosity, TEXT("* \tbundle, full size, current size, reserved, timestamp, age scale"))
+	UE_LOG(LogInstallBundleManager, Verbosity, TEXT("* \tbundle, full size, current size, diff, reserved, timestamp, age scale"))
 
 #define INSTALLBUNDLECACHE_CSV_LOG(Verbosity) \
 	if (Info.CurrentInstallSize > 0 || Info.State != ECacheState::Released) \
 	{ \
-		UE_LOG(LogInstallBundleManager, Verbosity, TEXT("* \t%s, %" UINT64_FMT ", %" UINT64_FMT ", %s, %s, %f"), *BundleName.ToString(), Info.FullInstallSize, Info.CurrentInstallSize, (Info.State == ECacheState::Reserved) ? TEXT("true") : TEXT("false"), *Info.TimeStamp.ToString(), Info.AgeScalar) \
+		const TCHAR* Diff = TEXT("="); \
+		if (Info.FullInstallSize > Info.CurrentInstallSize) Diff = TEXT(">"); \
+		else if(Info.FullInstallSize < Info.CurrentInstallSize) Diff = TEXT("<"); \
+		UE_LOG(LogInstallBundleManager, Verbosity, TEXT("* \t%s, %" UINT64_FMT ", %" UINT64_FMT ", %s, %s, %s, %f"), *BundleName.ToString(), Info.FullInstallSize, Info.CurrentInstallSize, Diff, (Info.State == ECacheState::Reserved) ? TEXT("true") : TEXT("false"), *Info.TimeStamp.ToString(), Info.AgeScalar) \
 	}
 
 	auto DumpToLog_Default = [](FName BundleName, const FBundleCacheInfo& Info)
