Removed InstallBundleManager private pch.
Timing when compiling module with PCH(including compiling PCH): 5.64 secs
Timing using shared PCH: 4.44s secs

diff --git a/Engine/Source/Runtime/InstallBundleManager/InstallBundleManager.Build.cs b/Engine/Source/Runtime/InstallBundleManager/InstallBundleManager.Build.cs
--- a/Engine/Source/Runtime/InstallBundleManager/InstallBundleManager.Build.cs
+++ b/Engine/Source/Runtime/InstallBundleManager/InstallBundleManager.Build.cs
@@ -6,8 +6,6 @@ public class InstallBundleManager : ModuleRules
 {
 	public InstallBundleManager(ReadOnlyTargetRules Target) : base(Target)
 	{
-		PrivatePCHHeaderFile = "Private/InstallBundleManagerPrivatePCH.h";
-
 		PrivateDependencyModuleNames.AddRange(
 			new string[] {
 				"Core",
