Renamed private pch headers in InstallBundleManager and fixed -NoPCH -DisableUnity compile issues.

diff --git a/Engine/Source/Runtime/InstallBundleManager/Private/BundlePrereqCombinedStatusHelper.cpp b/Engine/Source/Runtime/InstallBundleManager/Private/BundlePrereqCombinedStatusHelper.cpp
--- a/Engine/Source/Runtime/InstallBundleManager/Private/BundlePrereqCombinedStatusHelper.cpp
+++ b/Engine/Source/Runtime/InstallBundleManager/Private/BundlePrereqCombinedStatusHelper.cpp
@@ -2,7 +2,7 @@
 
 #include "BundlePrereqCombinedStatusHelper.h"
 #include "Containers/Ticker.h"
-#include "InstallBundleManagerPrivatePCH.h"
+#include "InstallBundleManagerPrivate.h"
 #include "Stats/Stats.h"
 #include "Algo/Transform.h"
 
diff --git a/Engine/Source/Runtime/InstallBundleManager/Private/InstallBundleCache.cpp b/Engine/Source/Runtime/InstallBundleManager/Private/InstallBundleCache.cpp
--- a/Engine/Source/Runtime/InstallBundleManager/Private/InstallBundleCache.cpp
+++ b/Engine/Source/Runtime/InstallBundleManager/Private/InstallBundleCache.cpp
@@ -1,7 +1,7 @@
 // Copyright Epic Games, Inc. All Rights Reserved.
 
 #include "InstallBundleCache.h"
-#include "InstallBundleManagerPrivatePCH.h"
+#include "InstallBundleManagerPrivate.h"
 
 #define INSTALLBUNDLE_CACHE_CHECK_INVARIANTS (DO_CHECK && 0)
 #define INSTALLBUNDLE_CACHE_DUMP_INFO (0)
diff --git a/Engine/Source/Runtime/InstallBundleManager/Private/InstallBundleManagerModule.cpp b/Engine/Source/Runtime/InstallBundleManager/Private/InstallBundleManagerModule.cpp
--- a/Engine/Source/Runtime/InstallBundleManager/Private/InstallBundleManagerModule.cpp
+++ b/Engine/Source/Runtime/InstallBundleManager/Private/InstallBundleManagerModule.cpp
@@ -1,6 +1,6 @@
 // Copyright Epic Games, Inc. All Rights Reserved.
 #include "InstallBundleManagerModule.h"
-#include "InstallBundleManagerPrivatePCH.h"
+#include "InstallBundleManagerPrivate.h"
 
 IMPLEMENT_MODULE(FInstallBundleManagerModule, InstallBundleManager);
 DEFINE_LOG_CATEGORY(LogInstallBundleManager);
diff --git a/Engine/Source/Runtime/InstallBundleManager/Private/InstallBundleManagerPrivate.h b/Engine/Source/Runtime/InstallBundleManager/Private/InstallBundleManagerPrivate.h
new file mode 100644
--- /dev/null
+++ b/Engine/Source/Runtime/InstallBundleManager/Private/InstallBundleManagerPrivate.h
@@ -0,0 +1,9 @@
+// Copyright Epic Games, Inc. All Rights Reserved.
+
+#pragma once
+
+#include "CoreMinimal.h"
+#include "ProfilingDebugging/CsvProfiler.h"
+
+DECLARE_LOG_CATEGORY_EXTERN(LogInstallBundleManager, Display, All);
+CSV_DECLARE_CATEGORY_MODULE_EXTERN(INSTALLBUNDLEMANAGER_API, InstallBundleManager);
diff --git a/Engine/Source/Runtime/InstallBundleManager/Private/InstallBundleManagerPrivatePCH.h b/Engine/Source/Runtime/InstallBundleManager/Private/InstallBundleManagerPrivatePCH.h
deleted file mode 100644
diff --git a/Engine/Source/Runtime/InstallBundleManager/Private/InstallBundleTypes.cpp b/Engine/Source/Runtime/InstallBundleManager/Private/InstallBundleTypes.cpp
--- a/Engine/Source/Runtime/InstallBundleManager/Private/InstallBundleTypes.cpp
+++ b/Engine/Source/Runtime/InstallBundleManager/Private/InstallBundleTypes.cpp
@@ -1,8 +1,8 @@
 // Copyright Epic Games, Inc. All Rights Reserved.
 
 #include "InstallBundleTypes.h"
-#include "InstallBundleManagerPrivatePCH.h"
 
+#include "InstallBundleManagerPrivate.h"
 #include "InstallBundleUtils.h"
 #include "Misc/CString.h"
 
diff --git a/Engine/Source/Runtime/InstallBundleManager/Private/InstallBundleUtils.cpp b/Engine/Source/Runtime/InstallBundleManager/Private/InstallBundleUtils.cpp
--- a/Engine/Source/Runtime/InstallBundleManager/Private/InstallBundleUtils.cpp
+++ b/Engine/Source/Runtime/InstallBundleManager/Private/InstallBundleUtils.cpp
@@ -1,7 +1,8 @@
 // Copyright Epic Games, Inc. All Rights Reserved.
 
 #include "InstallBundleUtils.h"
-#include "InstallBundleManagerPrivatePCH.h"
+
+#include "InstallBundleManagerPrivate.h"
 #include "Misc/App.h"
 
 #include "HAL/PlatformApplicationMisc.h"
