Adds InstallBundleSourceStreaming to support on-demand assets that are fetched directly from the CDN and written to a virtual file cache

diff --git a/Engine/Source/Runtime/InstallBundleManager/Private/InstallBundleTypes.cpp b/Engine/Source/Runtime/InstallBundleManager/Private/InstallBundleTypes.cpp
--- a/Engine/Source/Runtime/InstallBundleManager/Private/InstallBundleTypes.cpp
+++ b/Engine/Source/Runtime/InstallBundleManager/Private/InstallBundleTypes.cpp
@@ -17,6 +17,7 @@ const TCHAR* LexToString(EInstallBundleSourceType Type)
 		TEXT("Platform"),
 #endif // WITH_PLATFORM_INSTALL_BUNDLE_SOURCE
 		TEXT("GameCustom"),
+		TEXT("Streaming"),
 	};
 
 	return InstallBundleUtil::TLexToString(Type, Strings);
diff --git a/Engine/Source/Runtime/InstallBundleManager/Public/InstallBundleTypes.h b/Engine/Source/Runtime/InstallBundleManager/Public/InstallBundleTypes.h
--- a/Engine/Source/Runtime/InstallBundleManager/Public/InstallBundleTypes.h
+++ b/Engine/Source/Runtime/InstallBundleManager/Public/InstallBundleTypes.h
@@ -29,6 +29,7 @@ enum class EInstallBundleSourceType : int
 	Platform,
 #endif // WITH_PLATFORM_INSTALL_BUNDLE_SOURCE
 	GameCustom,
+	Streaming,
 	Count,
 };
 ENUM_RANGE_BY_COUNT(EInstallBundleSourceType, EInstallBundleSourceType::Count);
