Letting all request share the same context
- This change was made due to the Lazy notification was sent too early
#tests Android Test Client

diff --git a/Engine/Plugins/Experimental/DefaultInstallBundleManager/Source/Private/DefaultInstallBundleManager.cpp b/Engine/Plugins/Experimental/DefaultInstallBundleManager/Source/Private/DefaultInstallBundleManager.cpp
--- a/Engine/Plugins/Experimental/DefaultInstallBundleManager/Source/Private/DefaultInstallBundleManager.cpp
+++ b/Engine/Plugins/Experimental/DefaultInstallBundleManager/Source/Private/DefaultInstallBundleManager.cpp
@@ -2987,7 +2987,9 @@ EInstallBundleManagerInitState FDefaultInstallBundleManager::GetInitState() cons
 }
 
 TValueOrError<FInstallBundleRequestInfo, EInstallBundleResult> FDefaultInstallBundleManager::RequestUpdateContent(
-	TArrayView<const FName> InBundleNames, EInstallBundleRequestFlags Flags, ELogVerbosity::Type LogVerbosityOverride /*= ELogVerbosity::NoLogging*/)
+	TArrayView<const FName> InBundleNames, EInstallBundleRequestFlags Flags,
+	ELogVerbosity::Type LogVerbosityOverride /*= ELogVerbosity::NoLogging*/,
+	InstallBundleUtil::FContentRequestSharedContextPtr RequestSharedContext /*= nullptr*/)
 {
 	CSV_SCOPED_TIMING_STAT(InstallBundleManager, InstallBundleManager_RequestUpdateContent);
 
@@ -3005,9 +3007,6 @@ TValueOrError<FInstallBundleRequestInfo, EInstallBundleResult> FDefaultInstallBu
 		return MakeError(EInstallBundleResult::InitializationPending);
 	}
 
-	// Shared between all bundles for the same request
-	InstallBundleUtil::FContentRequestSharedContextPtr RequestSharedContext;
-
 	TSet<FName> BundlesToLoad = GatherBundlesForRequest(InBundleNames, RetInfo.InfoFlags);
 	for (const FName& BundleName : BundlesToLoad)
 	{
diff --git a/Engine/Plugins/Experimental/DefaultInstallBundleManager/Source/Public/DefaultInstallBundleManager.h b/Engine/Plugins/Experimental/DefaultInstallBundleManager/Source/Public/DefaultInstallBundleManager.h
--- a/Engine/Plugins/Experimental/DefaultInstallBundleManager/Source/Public/DefaultInstallBundleManager.h
+++ b/Engine/Plugins/Experimental/DefaultInstallBundleManager/Source/Public/DefaultInstallBundleManager.h
@@ -615,7 +615,7 @@ public:
 
 	virtual EInstallBundleManagerInitState GetInitState() const override;
 
-	virtual TValueOrError<FInstallBundleRequestInfo, EInstallBundleResult> RequestUpdateContent(TArrayView<const FName> InBundleNames, EInstallBundleRequestFlags Flags, ELogVerbosity::Type LogVerbosityOverride = ELogVerbosity::NoLogging) override;
+	virtual TValueOrError<FInstallBundleRequestInfo, EInstallBundleResult> RequestUpdateContent(TArrayView<const FName> InBundleNames, EInstallBundleRequestFlags Flags, ELogVerbosity::Type LogVerbosityOverride = ELogVerbosity::NoLogging, InstallBundleUtil::FContentRequestSharedContextPtr RequestSharedContext = nullptr) override;
 
 	virtual FDelegateHandle GetContentState(TArrayView<const FName> InBundleNames, EInstallBundleGetContentStateFlags Flags, bool bAddDependencies, FInstallBundleGetContentStateDelegate Callback, FName RequestTag = TEXT("None")) override;
 
@@ -749,7 +749,6 @@ protected:
 	TSharedRef<InstallBundleManagerUtil::FPersistentStatContainer> PersistentStats;
 
 	TArray<TUniquePtr<InstallBundleUtil::FInstallBundleTask>> AsyncMountTasks;
-
 	
 	bool bIsCheckingForPatch = false;
 	bool bDelayCheckingForContentPatch = false;
diff --git a/Engine/Source/Runtime/InstallBundleManager/Private/InstallBundleManagerInterface.cpp b/Engine/Source/Runtime/InstallBundleManager/Private/InstallBundleManagerInterface.cpp
--- a/Engine/Source/Runtime/InstallBundleManager/Private/InstallBundleManagerInterface.cpp
+++ b/Engine/Source/Runtime/InstallBundleManager/Private/InstallBundleManagerInterface.cpp
@@ -54,9 +54,9 @@ const TSharedPtr<IInstallBundleSource> IInstallBundleManager::GetBundleSource(EI
 }
 
 
-TValueOrError<FInstallBundleRequestInfo, EInstallBundleResult> IInstallBundleManager::RequestUpdateContent(FName BundleName, EInstallBundleRequestFlags Flags, ELogVerbosity::Type LogVerbosityOverride /*= ELogVerbosity::NoLogging*/)
+TValueOrError<FInstallBundleRequestInfo, EInstallBundleResult> IInstallBundleManager::RequestUpdateContent(FName BundleName, EInstallBundleRequestFlags Flags, ELogVerbosity::Type LogVerbosityOverride /*= ELogVerbosity::NoLogging*/,  InstallBundleUtil::FContentRequestSharedContextPtr RequestSharedContext /*= nullptr*/)
 {
-	return RequestUpdateContent(MakeArrayView(&BundleName, 1), Flags, LogVerbosityOverride);
+	return RequestUpdateContent(MakeArrayView(&BundleName, 1), Flags, LogVerbosityOverride, RequestSharedContext);
 }
 
 FDelegateHandle IInstallBundleManager::GetContentState(FName BundleName, EInstallBundleGetContentStateFlags Flags, bool bAddDependencies, FInstallBundleGetContentStateDelegate Callback, FName RequestTag /*= NAME_None*/)
diff --git a/Engine/Source/Runtime/InstallBundleManager/Public/InstallBundleManagerInterface.h b/Engine/Source/Runtime/InstallBundleManager/Public/InstallBundleManagerInterface.h
--- a/Engine/Source/Runtime/InstallBundleManager/Public/InstallBundleManagerInterface.h
+++ b/Engine/Source/Runtime/InstallBundleManager/Public/InstallBundleManagerInterface.h
@@ -24,6 +24,12 @@
 class IAnalyticsProviderET;
 class IInstallBundleSource;
 
+namespace InstallBundleUtil
+{
+	struct FContentRequestSharedContext;
+	using FContentRequestSharedContextPtr = TSharedPtr<FContentRequestSharedContext>;
+}
+
 struct FInstallBundleProgress
 {
 	FName BundleName;
@@ -113,8 +119,8 @@ public:
 
 	virtual EInstallBundleManagerInitState GetInitState() const = 0;
 
-	TValueOrError<FInstallBundleRequestInfo, EInstallBundleResult> RequestUpdateContent(FName BundleName, EInstallBundleRequestFlags Flags, ELogVerbosity::Type LogVerbosityOverride = ELogVerbosity::NoLogging);
-	virtual TValueOrError<FInstallBundleRequestInfo, EInstallBundleResult> RequestUpdateContent(TArrayView<const FName> BundleNames, EInstallBundleRequestFlags Flags, ELogVerbosity::Type LogVerbosityOverride = ELogVerbosity::NoLogging) = 0;
+	TValueOrError<FInstallBundleRequestInfo, EInstallBundleResult> RequestUpdateContent(FName BundleName, EInstallBundleRequestFlags Flags, ELogVerbosity::Type LogVerbosityOverride = ELogVerbosity::NoLogging, InstallBundleUtil::FContentRequestSharedContextPtr RequestSharedContext = nullptr);
+	virtual TValueOrError<FInstallBundleRequestInfo, EInstallBundleResult> RequestUpdateContent(TArrayView<const FName> BundleNames, EInstallBundleRequestFlags Flags, ELogVerbosity::Type LogVerbosityOverride = ELogVerbosity::NoLogging, InstallBundleUtil::FContentRequestSharedContextPtr RequestSharedContext = nullptr) = 0;
 
 	FDelegateHandle GetContentState(FName BundleName, EInstallBundleGetContentStateFlags Flags, bool bAddDependencies, FInstallBundleGetContentStateDelegate Callback, FName RequestTag = NAME_None);
 	virtual FDelegateHandle GetContentState(TArrayView<const FName> BundleNames, EInstallBundleGetContentStateFlags Flags, bool bAddDependencies, FInstallBundleGetContentStateDelegate Callback, FName RequestTag = NAME_None) = 0;
diff --git a/Engine/Source/Runtime/NullInstallBundleManager/Private/NullInstallBundleManager.cpp b/Engine/Source/Runtime/NullInstallBundleManager/Private/NullInstallBundleManager.cpp
--- a/Engine/Source/Runtime/NullInstallBundleManager/Private/NullInstallBundleManager.cpp
+++ b/Engine/Source/Runtime/NullInstallBundleManager/Private/NullInstallBundleManager.cpp
@@ -33,7 +33,7 @@ class FNullInstallBundleManager : public IInstallBundleManager
 		return EInstallBundleManagerInitState::Succeeded;
 	}
 
-	virtual TValueOrError<FInstallBundleRequestInfo, EInstallBundleResult> RequestUpdateContent(TArrayView<const FName> BundleNames, EInstallBundleRequestFlags Flags, ELogVerbosity::Type LogVerbosityOverride = ELogVerbosity::NoLogging) override
+	virtual TValueOrError<FInstallBundleRequestInfo, EInstallBundleResult> RequestUpdateContent(TArrayView<const FName> BundleNames, EInstallBundleRequestFlags Flags, ELogVerbosity::Type LogVerbosityOverride = ELogVerbosity::NoLogging, InstallBundleUtil::FContentRequestSharedContextPtr RequestSharedContext = nullptr) override
 	{
 		return MakeValue(FInstallBundleRequestInfo());
 	}
