FInstallBundleCombinedProgressTracker (InstallBundleManager Utility)
- Fixed issue when we could en up showing 100% Updating status when there was still Finishing work being done while waiting to receive GetBundleProgress. Now if ProgressPercent is 1.0f we will correctly show Finilizing progress if applicable.
- No longer track 0 weight bundles as they have no meaningful reason to be tracked and can lead to errors.
#tests Android Dev Client

diff --git a/Engine/Source/Runtime/InstallBundleManager/Private/BundlePrereqCombinedStatusHelper.cpp b/Engine/Source/Runtime/InstallBundleManager/Private/BundlePrereqCombinedStatusHelper.cpp
--- a/Engine/Source/Runtime/InstallBundleManager/Private/BundlePrereqCombinedStatusHelper.cpp
+++ b/Engine/Source/Runtime/InstallBundleManager/Private/BundlePrereqCombinedStatusHelper.cpp
@@ -120,6 +120,12 @@ void FInstallBundleCombinedProgressTracker::SetBundlesToTrackFromContentState(co
 		const FInstallBundleContentState* BundleState = BundleContentState.IndividualBundleStates.Find(Bundle);
 		if (ensureAlwaysMsgf(BundleState, TEXT("Trying to track unknown bundle %s"), *Bundle.ToString()))
 		{
+			//Filter out any bundles with effectively 0 weight
+			if (BundleState->Weight <= SMALL_NUMBER)
+			{
+				continue;
+			}
+
 			//Track if we need any kind of bundle updates
 			if (BundleState->State == EInstallBundleInstallState::NotInstalled || BundleState->State == EInstallBundleInstallState::NeedsUpdate)
 			{
@@ -193,9 +199,12 @@ void FInstallBundleCombinedProgressTracker::UpdateCombinedStatus()
 	//if we don't yet have a bundle status cache entry for a particular requirement
 	//then we can't yet tell what work is required on that bundle yet. We need to go ahead and make sure we don't
 	//show a status like "Installed" before we know what state that bundle is in. Make sure we show at LEAST
-	//updating in that case, so start with Downloading since that is the first Updating case
+	//updating in that case, so start with Downloading since that is the first Updating case.
+	//However if all bundle progress is finished, don't just sit showing 100% and Updating when we could potentially
+	//be showing Finishing progress
 	if ((BundleStatusCache.Num() < RequiredBundleNames.Num())
-		&& (BundleStatusCache.Num() > 0))
+		&& (BundleStatusCache.Num() > 0)
+		&& (CurrentCombinedProgress.ProgressPercent < 1.0f))
 	{
 		EarliestBundleState = EInstallBundleStatus::Updating;
 	}
