Change use of TAtomic to std::atomic

#tests Compiled

diff --git a/Engine/Source/Runtime/InstallBundleManager/Private/InstallBundleUtils.cpp b/Engine/Source/Runtime/InstallBundleManager/Private/InstallBundleUtils.cpp
--- a/Engine/Source/Runtime/InstallBundleManager/Private/InstallBundleUtils.cpp
+++ b/Engine/Source/Runtime/InstallBundleManager/Private/InstallBundleUtils.cpp
@@ -78,7 +78,7 @@ namespace InstallBundleUtil
 		}
 	}
 
-	TAtomic<int32> InstallBundleSuppressAnalyticsCounter = 0;
+	std::atomic<int32> InstallBundleSuppressAnalyticsCounter = 0;
 
 	FInstallBundleSuppressAnalytics::FInstallBundleSuppressAnalytics()
 		: bIsEnabled(false)
