Refactored hermes load context to use link codes in automated tests and islands.
Fixed automated tests for download on demand.

diff --git a/Engine/Source/Runtime/InstallBundleManager/Public/InstallBundleManagerModule.h b/Engine/Source/Runtime/InstallBundleManager/Public/InstallBundleManagerModule.h
--- a/Engine/Source/Runtime/InstallBundleManager/Public/InstallBundleManagerModule.h
+++ b/Engine/Source/Runtime/InstallBundleManager/Public/InstallBundleManagerModule.h
@@ -38,6 +38,8 @@ public:
 	virtual void PreUnloadCallback() override
 	{
 		InstallBundleManager = nullptr;
+		// make sure the bundle manager was cleaned up
+		check(LastInstallBundleManager.IsValid() == false);
 	}
 
 	TSharedPtr<IInstallBundleManager> GetInstallBundleManager()
