Fix for not correctly tracking combined progress if all supplied bundles had 0 weight. Will now evenly distribute progress if all bundles have 0 weight instead of tracking nothing.
#tests Android Dev Client

diff --git a/Engine/Source/Runtime/InstallBundleManager/Private/BundlePrereqCombinedStatusHelper.cpp b/Engine/Source/Runtime/InstallBundleManager/Private/BundlePrereqCombinedStatusHelper.cpp
--- a/Engine/Source/Runtime/InstallBundleManager/Private/BundlePrereqCombinedStatusHelper.cpp
+++ b/Engine/Source/Runtime/InstallBundleManager/Private/BundlePrereqCombinedStatusHelper.cpp
@@ -113,6 +113,25 @@ void FInstallBundleCombinedProgressTracker::SetBundlesToTrackFromContentState(co
 	CachedBundleWeights.Empty();
 	BundleStatusCache.Empty();
 
+	//Go through all bundles until we hit a non-zero weight bundle. 
+	//This is to help catch instances where we pass in all zero weight bundles to track and need
+	//to thus calculate their weight dynamically based on everything having even weight
+	bool bAreAllBundlesZeroWeight = true;
+	for (const TPair<FName, FInstallBundleContentState>& IndividualBundlePair : BundleContentState.IndividualBundleStates)
+	{
+		const FInstallBundleContentState& BundleState = IndividualBundlePair.Value;
+		if (BundleState.Weight <= SMALL_NUMBER)
+		{
+			continue;
+		}
+		else
+		{
+			bAreAllBundlesZeroWeight = false;
+			break;
+		}
+	}
+		
+
 	bool bBundleNeedsUpdate = false;
 	float TotalWeight = 0.0f;
 	for (const FName& Bundle : BundlesToTrack)
@@ -120,8 +139,8 @@ void FInstallBundleCombinedProgressTracker::SetBundlesToTrackFromContentState(co
 		const FInstallBundleContentState* BundleState = BundleContentState.IndividualBundleStates.Find(Bundle);
 		if (ensureAlwaysMsgf(BundleState, TEXT("Trying to track unknown bundle %s"), *Bundle.ToString()))
 		{
-			//Filter out any bundles with effectively 0 weight
-			if (BundleState->Weight <= SMALL_NUMBER)
+			//Filter out any bundles with effectively 0 weight (unless all bundles are 0 weight)
+			if (!bAreAllBundlesZeroWeight && (BundleState->Weight <= SMALL_NUMBER))
 			{
 				continue;
 			}
@@ -134,7 +153,10 @@ void FInstallBundleCombinedProgressTracker::SetBundlesToTrackFromContentState(co
 
 			//Save required bundles and their weights
 			RequiredBundleNames.Add(Bundle);
-			CachedBundleWeights.FindOrAdd(Bundle) = BundleState->Weight;
+
+			//If all bundles are zero weight, just treat this weight as 1 so everything ends up with 1 weight and is evenly distributed
+			CachedBundleWeights.FindOrAdd(Bundle) = bAreAllBundlesZeroWeight ? 1.0f : BundleState->Weight;
+
 			TotalWeight += BundleState->Weight;
 		}
 	}
