Refactor bundle regex matching into shared utility functions.

diff --git a/Engine/Plugins/Experimental/DefaultInstallBundleManager/Source/Private/InstallBundleManagerUtil.cpp b/Engine/Plugins/Experimental/DefaultInstallBundleManager/Source/Private/InstallBundleManagerUtil.cpp
--- a/Engine/Plugins/Experimental/DefaultInstallBundleManager/Source/Private/InstallBundleManagerUtil.cpp
+++ b/Engine/Plugins/Experimental/DefaultInstallBundleManager/Source/Private/InstallBundleManagerUtil.cpp
@@ -14,8 +14,6 @@
 	#define WITH_PLATFORM_INSTALL_BUNDLE_SOURCE 0
 #endif
 
-#include "IAnalyticsProviderET.h"
-
 DEFAULTINSTALLBUNDLEMANAGER_API FString GInstallBundleManagerIni;
 
 DEFINE_LOG_CATEGORY(LogDefaultInstallBundleManager);
diff --git a/Engine/Plugins/Experimental/DefaultInstallBundleManager/Source/Private/InstallBundleSourcePlatformBase.cpp b/Engine/Plugins/Experimental/DefaultInstallBundleManager/Source/Private/InstallBundleSourcePlatformBase.cpp
--- a/Engine/Plugins/Experimental/DefaultInstallBundleManager/Source/Private/InstallBundleSourcePlatformBase.cpp
+++ b/Engine/Plugins/Experimental/DefaultInstallBundleManager/Source/Private/InstallBundleSourcePlatformBase.cpp
@@ -52,50 +52,7 @@ FInstallBundleSourceInitInfo FInstallBundleSourcePlatformBase::Init(
 		return InitInfo;
 	}
 
-	for (const TPair<FString, FConfigSection>& Pair : *InstallBundleConfig)
-	{
-		const FString& Section = Pair.Key;
-		if (!Section.StartsWith(InstallBundleUtil::GetInstallBundleSectionPrefix()))
-			continue;
-
-		if (!IsPlatformInstallBundleConfig(InstallBundleConfig, Section))
-			continue;
-
-		TArray<FString> StrSearchRegexPatterns;
-		if (!InstallBundleConfig->GetArray(*Section, TEXT("FileRegex"), StrSearchRegexPatterns))
-			continue;
-
-		TArray<FRegexPattern> SearchRegexPatterns;
-		SearchRegexPatterns.Reserve(StrSearchRegexPatterns.Num());
-		for (const FString& Str : StrSearchRegexPatterns)
-		{
-			SearchRegexPatterns.Emplace(Str, ERegexPatternFlags::CaseInsensitive);
-		}
-
-		const FString BundleName = Section.RightChop(InstallBundleUtil::GetInstallBundleSectionPrefix().Len());
-		BundleRegexList.Emplace(TPair<FString, TArray<FRegexPattern>>(BundleName, MoveTemp(SearchRegexPatterns)));
-	}
-
-	BundleRegexList.StableSort([](const TPair<FString, TArray<FRegexPattern>>& PairA, const TPair<FString, TArray<FRegexPattern>>& PairB) -> bool
-	{
-		int BundleAOrder = INT_MAX;
-		int BundleBOrder = INT_MAX;
-
-		const FString SectionA = InstallBundleUtil::GetInstallBundleSectionPrefix() + PairA.Key;
-		const FString SectionB = InstallBundleUtil::GetInstallBundleSectionPrefix() + PairB.Key;
-
-		if (!GConfig->GetInt(*SectionA, TEXT("Order"), BundleAOrder, GInstallBundleIni))
-		{
-			LOG_SOURCE_PLATFORMBASE(Warning, TEXT("Bundle Section %s doesn't have an order"), *SectionA);
-		}
-
-		if (!GConfig->GetInt(*SectionB, TEXT("Order"), BundleBOrder, GInstallBundleIni))
-		{
-			LOG_SOURCE_PLATFORMBASE(Warning, TEXT("Bundle Section %s doesn't have an order"), *SectionB);
-		}
-
-		return BundleAOrder < BundleBOrder;
-	});
+	BundleRegexList = InstallBundleUtil::LoadBundleRegexFromConfig(*InstallBundleConfig, InstallBundleUtil::IsPlatformInstallBundlePredicate);
 
 	return InitInfo;
 }
diff --git a/Engine/Plugins/Experimental/DefaultInstallBundleManager/Source/Public/InstallBundleSourcePlatformBase.h b/Engine/Plugins/Experimental/DefaultInstallBundleManager/Source/Public/InstallBundleSourcePlatformBase.h
--- a/Engine/Plugins/Experimental/DefaultInstallBundleManager/Source/Public/InstallBundleSourcePlatformBase.h
+++ b/Engine/Plugins/Experimental/DefaultInstallBundleManager/Source/Public/InstallBundleSourcePlatformBase.h
@@ -39,7 +39,6 @@ public:
 
 
 protected:
-	virtual bool IsPlatformInstallBundleConfig(const FConfigFile* InstallBundleConfig, const FString& Section) const = 0;
 	virtual bool QueryPersistentBundleInfo(FInstallBundleSourcePersistentBundleInfo& SourceBundleInfo) const = 0;
 
 	template<typename T>
diff --git a/Engine/Source/Runtime/InstallBundleManager/Private/InstallBundleUtils.cpp b/Engine/Source/Runtime/InstallBundleManager/Private/InstallBundleUtils.cpp
--- a/Engine/Source/Runtime/InstallBundleManager/Private/InstallBundleUtils.cpp
+++ b/Engine/Source/Runtime/InstallBundleManager/Private/InstallBundleUtils.cpp
@@ -17,6 +17,9 @@
 #include "Serialization/JsonSerializerMacros.h"
 #include "Stats/Stats.h"
 
+#include "Algo/AnyOf.h"
+#include "Algo/Find.h"
+
 namespace InstallBundleUtil
 {
 	FString GetAppVersion()
@@ -52,6 +55,100 @@ namespace InstallBundleUtil
 		return Prefix;
 	}
 
+	bool AllInstallBundlePredicate(const FConfigFile& InstallBundleConfig, const FString& Section) 
+	{ 
+		return true; 
+	}
+
+	bool IsPlatformInstallBundlePredicate(const FConfigFile& InstallBundleConfig, const FString& Section)
+	{
+		int32 ChunkID = 0;
+		if (InstallBundleConfig.GetInt(*Section, TEXT("PlatformChunkID"), ChunkID) && ChunkID < 0)
+		{
+			return false;
+		}
+
+		return true;
+	}
+
+	TArray<TPair<FString, TArray<FRegexPattern>>> LoadBundleRegexFromConfig(
+		const FConfigFile& InstallBundleConfig, 
+		TFunctionRef<bool(const FConfigFile& InstallBundleConfig, const FString& Section)> SectionPredicate /*= AllInstallBundlePredicate*/)
+	{
+		TArray<TPair<FString, TArray<FRegexPattern>>> BundleRegexList; // BundleName -> FileRegex
+
+		for (const TPair<FString, FConfigSection>& Pair : InstallBundleConfig)
+		{
+			const FString& Section = Pair.Key;
+			if (!Section.StartsWith(InstallBundleUtil::GetInstallBundleSectionPrefix()))
+				continue;
+
+			if (!SectionPredicate(InstallBundleConfig, Section))
+				continue;
+
+			TArray<FString> StrSearchRegexPatterns;
+			if (!InstallBundleConfig.GetArray(*Section, TEXT("FileRegex"), StrSearchRegexPatterns))
+				continue;
+
+			TArray<FRegexPattern> SearchRegexPatterns;
+			SearchRegexPatterns.Reserve(StrSearchRegexPatterns.Num());
+			for (const FString& Str : StrSearchRegexPatterns)
+			{
+				SearchRegexPatterns.Emplace(Str, ERegexPatternFlags::CaseInsensitive);
+			}
+
+			const FString BundleName = Section.RightChop(InstallBundleUtil::GetInstallBundleSectionPrefix().Len());
+			BundleRegexList.Emplace(TPair<FString, TArray<FRegexPattern>>(BundleName, MoveTemp(SearchRegexPatterns)));
+		}
+
+		BundleRegexList.StableSort([](const TPair<FString, TArray<FRegexPattern>>& PairA, const TPair<FString, TArray<FRegexPattern>>& PairB) -> bool
+		{
+			int BundleAOrder = INT_MAX;
+			int BundleBOrder = INT_MAX;
+
+			const FString SectionA = InstallBundleUtil::GetInstallBundleSectionPrefix() + PairA.Key;
+			const FString SectionB = InstallBundleUtil::GetInstallBundleSectionPrefix() + PairB.Key;
+
+			if (!GConfig->GetInt(*SectionA, TEXT("Order"), BundleAOrder, GInstallBundleIni))
+			{
+				UE_LOG(LogInstallBundleManager, Warning, TEXT("Bundle Section %s doesn't have an order"), *SectionA);
+			}
+
+			if (!GConfig->GetInt(*SectionB, TEXT("Order"), BundleBOrder, GInstallBundleIni))
+			{
+				UE_LOG(LogInstallBundleManager, Warning, TEXT("Bundle Section %s doesn't have an order"), *SectionB);
+			}
+
+			return BundleAOrder < BundleBOrder;
+		});
+
+		return BundleRegexList;
+	}
+
+	bool MatchBundleRegex(
+		const TArray<TPair<FString, TArray<FRegexPattern>>>& BundleRegexList,
+		const FString& Path,
+		FString& OutBundleName)
+	{
+		const TPair<FString, TArray<FRegexPattern>>* BundleRegexPair = Algo::FindByPredicate(BundleRegexList,
+			[&Path](const TPair<FString, TArray<FRegexPattern>>& Pair)
+			{
+				const TArray<FRegexPattern>& SearchRegexPatterns = Pair.Value;
+				return Algo::AnyOf(SearchRegexPatterns, [&Path](const FRegexPattern& Pattern)
+				{
+					return FRegexMatcher(Pattern, Path).FindNext();
+				});
+			});
+
+		if (BundleRegexPair)
+		{
+			OutBundleName = BundleRegexPair->Key;
+			return true;
+		}
+
+		return false;
+	}
+
 	FName FInstallBundleManagerKeepAwake::Tag(TEXT("InstallBundleManagerKeepAwake"));
 	FName FInstallBundleManagerKeepAwake::TagWithRendering(TEXT("InstallBundleManagerKeepAwakeWithRendering"));
 
diff --git a/Engine/Source/Runtime/InstallBundleManager/Public/InstallBundleUtils.h b/Engine/Source/Runtime/InstallBundleManager/Public/InstallBundleUtils.h
--- a/Engine/Source/Runtime/InstallBundleManager/Public/InstallBundleUtils.h
+++ b/Engine/Source/Runtime/InstallBundleManager/Public/InstallBundleUtils.h
@@ -22,6 +22,7 @@
 #include "Templates/UniquePtr.h"
 #include "Templates/UnrealTemplate.h"
 #include "UObject/NameTypes.h"
+#include "Internationalization/Regex.h"
 
 class FQueuedThreadPool;
 
@@ -36,6 +37,24 @@ namespace InstallBundleUtil
 
 	INSTALLBUNDLEMANAGER_API const FString& GetInstallBundleSectionPrefix();
 
+	// Returns true
+	INSTALLBUNDLEMANAGER_API bool AllInstallBundlePredicate(const FConfigFile& InstallBundleConfig, const FString& Section);
+
+	// Returns true if this bundle should be included in a platform package
+	INSTALLBUNDLEMANAGER_API bool IsPlatformInstallBundlePredicate(const FConfigFile& InstallBundleConfig, const FString& Section);
+
+	// Returns an ordered regex list used to map files to bundles
+	// SectionPredicate can be used to exclude non-relavent bundles from the list
+	INSTALLBUNDLEMANAGER_API TArray<TPair<FString, TArray<FRegexPattern>>> LoadBundleRegexFromConfig(
+		const FConfigFile& InstallBundleConfig,
+		TFunctionRef<bool(const FConfigFile& InstallBundleConfig, const FString& Section)> SectionPredicate = AllInstallBundlePredicate);
+
+	// Finds a matching install bundle for a given path
+	INSTALLBUNDLEMANAGER_API bool MatchBundleRegex(
+		const TArray<TPair<FString, TArray<FRegexPattern>>>& BundleRegexList,
+		const FString& Path,
+		FString& OutBundleName);
+
 	constexpr float MinimumBundleWeight = 0.05f;
 
 	// It would really be nice to have these in core
