Fixup the install bundle to use the right config for bundle sorting

diff --git a/Engine/Source/Runtime/InstallBundleManager/Private/InstallBundleUtils.cpp b/Engine/Source/Runtime/InstallBundleManager/Private/InstallBundleUtils.cpp
--- a/Engine/Source/Runtime/InstallBundleManager/Private/InstallBundleUtils.cpp
+++ b/Engine/Source/Runtime/InstallBundleManager/Private/InstallBundleUtils.cpp
@@ -115,20 +115,20 @@ namespace InstallBundleUtil
 			BundleRegexList.Emplace(TPair<FString, TArray<FRegexPattern>>(BundleName, MoveTemp(SearchRegexPatterns)));
 		}
 
-		BundleRegexList.StableSort([](const TPair<FString, TArray<FRegexPattern>>& PairA, const TPair<FString, TArray<FRegexPattern>>& PairB) -> bool
+		BundleRegexList.StableSort([InstallBundleConfig](const TPair<FString, TArray<FRegexPattern>>& PairA, const TPair<FString, TArray<FRegexPattern>>& PairB) -> bool
 		{
-			int BundleAOrder = INT_MAX;
-			int BundleBOrder = INT_MAX;
+			int32 BundleAOrder = INT_MAX;
+			int32 BundleBOrder = INT_MAX;
 
 			const FString SectionA = InstallBundleUtil::GetInstallBundleSectionPrefix() + PairA.Key;
 			const FString SectionB = InstallBundleUtil::GetInstallBundleSectionPrefix() + PairB.Key;
 
-			if (!GConfig->GetInt(*SectionA, TEXT("Order"), BundleAOrder, GInstallBundleIni))
+			if (!InstallBundleConfig.GetInt(*SectionA, TEXT("Order"), BundleAOrder))
 			{
 				UE_LOG(LogInstallBundleManager, Warning, TEXT("Bundle Section %s doesn't have an order"), *SectionA);
 			}
 
-			if (!GConfig->GetInt(*SectionB, TEXT("Order"), BundleBOrder, GInstallBundleIni))
+			if (!InstallBundleConfig.GetInt(*SectionB, TEXT("Order"), BundleBOrder))
 			{
 				UE_LOG(LogInstallBundleManager, Warning, TEXT("Bundle Section %s doesn't have an order"), *SectionB);
 			}
@@ -1408,4 +1408,4 @@ namespace InstallBundleUtil
 			return SessionPersistentStatMap.Find(SessionName);
 		}
 	}
-}
+}
\ No newline at end of file
