Engine changes for GFP download throuput optimization.

GFP state machine
- Avoid async GetContentState calls if possible.  These waste time and delay getting initial HTTP requests queued. The allows all the requests to be queued on the same tick.

BPS download service
- Move HTTP callbacks off game thread. This doesn't have much effect yet since the download service is still ticked on the game thread.  I'll follow this up with changes to fix that.

Install bundle stats
- Additional ensures to help catch operations happening out of order. I caught and fixed an issue where VK bundle source was sending completion callbacks too early.

diff --git a/Engine/Plugins/Experimental/GameFeatures/Source/GameFeatures/Private/GameFeaturePluginStateMachine.cpp b/Engine/Plugins/Experimental/GameFeatures/Source/GameFeatures/Private/GameFeaturePluginStateMachine.cpp
--- a/Engine/Plugins/Experimental/GameFeatures/Source/GameFeatures/Private/GameFeaturePluginStateMachine.cpp
+++ b/Engine/Plugins/Experimental/GameFeatures/Source/GameFeatures/Private/GameFeaturePluginStateMachine.cpp
@@ -1382,7 +1382,29 @@ struct FGameFeaturePluginState_Downloading : public FGameFeaturePluginState
 		TSharedPtr<IInstallBundleManager> BundleManager = IInstallBundleManager::GetPlatformInstallBundleManager();
 		const TArray<FName>& InstallBundles = StateProperties.ProtocolMetadata.GetSubtype<FInstallBundlePluginProtocolMetaData>().InstallBundles;
 
-		GotContentStateHandle = BundleManager->GetContentState(InstallBundles, EInstallBundleGetContentStateFlags::None, true, FInstallBundleGetContentStateDelegate::CreateRaw(this, &FGameFeaturePluginState_Downloading::OnGotContentState));
+		if (InstallBundles.Num() > 1)
+		{
+			GotContentStateHandle = BundleManager->GetContentState(InstallBundles, EInstallBundleGetContentStateFlags::None, true, 
+				FInstallBundleGetContentStateDelegate::CreateRaw(this, &FGameFeaturePluginState_Downloading::OnGotContentState));
+		}
+		else
+		{
+			// We usualy only have a use case of one bundle per GFP and we only care about relative weighting here, so we don't need 
+			// any of the other content state metadata. We can just assume the weight is 1.0 and not have to wait for the full 
+			// async call to get the rest of the metadata.
+			TValueOrError<FInstallBundleCombinedInstallState, EInstallBundleResult> MaybeInstallState = BundleManager->GetInstallStateSynchronous(InstallBundles, true);
+			check(MaybeInstallState.HasValue());
+			const FInstallBundleCombinedInstallState& InstallState = MaybeInstallState.GetValue();
+			FInstallBundleCombinedContentState HackContentState;
+			HackContentState.IndividualBundleStates.Reserve(InstallState.IndividualBundleStates.Num());
+			for(const TPair<FName, EInstallBundleInstallState>& Pair : InstallState.IndividualBundleStates)
+			{
+				FInstallBundleContentState& BundleContentState = HackContentState.IndividualBundleStates.Emplace(Pair.Key);
+				BundleContentState.State = Pair.Value;
+				BundleContentState.Weight = 1.0f;
+			}
+			OnGotContentState(MoveTemp(HackContentState));
+		}
 	}
 
 	virtual void UpdateState(FGameFeaturePluginStateStatus& StateStatus) override
diff --git a/Engine/Source/Runtime/InstallBundleManager/Private/InstallBundleUtils.cpp b/Engine/Source/Runtime/InstallBundleManager/Private/InstallBundleUtils.cpp
--- a/Engine/Source/Runtime/InstallBundleManager/Private/InstallBundleUtils.cpp
+++ b/Engine/Source/Runtime/InstallBundleManager/Private/InstallBundleUtils.cpp
@@ -18,6 +18,7 @@
 #include "Stats/Stats.h"
 
 #include "Algo/AnyOf.h"
+#include "Algo/AllOf.h"
 #include "Algo/Find.h"
 
 namespace InstallBundleUtil
@@ -277,7 +278,7 @@ namespace InstallBundleUtil
 	void FContentRequestStatsMap::StatsBegin(FName BundleName)
 	{
 		FContentRequestStats& Stats = StatsMap.FindOrAdd(BundleName);
-		if (ensureAlwaysMsgf(Stats.bOpen, TEXT("StatsBegin - Stat closed for %s"), *BundleName.ToString()) == false)
+		if (false == ensureAlwaysMsgf(Stats.bOpen, TEXT("StatsBegin - Stat closed for %s"), *BundleName.ToString()))
 		{
 			Stats = FContentRequestStats();
 		}
@@ -289,8 +290,14 @@ namespace InstallBundleUtil
 	{
 		FContentRequestStats& Stats = StatsMap.FindOrAdd(BundleName);
 
-		if (ensureAlwaysMsgf(Stats.bOpen, TEXT("StatsEnd - Stat closed for %s"), *BundleName.ToString()))
+		ensureAlwaysMsgf(Stats.bOpen && Stats.StartTime > 0, TEXT("StatsEnd - Stat closed for %s"), *BundleName.ToString());
+		if (Stats.bOpen)
 		{
+			ensureAlwaysMsgf(
+				Algo::AllOf(Stats.StateStats, 
+					[](const TPair<FString, FContentRequestStateStats>& Pair) { return !Pair.Value.bOpen; }),
+				TEXT("StatsEnd - StateStat open for %s"), *BundleName.ToString());
+
 			Stats.EndTime = FPlatformTime::Seconds();
 			Stats.bOpen = false;
 		}
@@ -298,20 +305,29 @@ namespace InstallBundleUtil
 
 	void FContentRequestStatsMap::StatsReset(FName BundleName)
 	{
-		StatsMap.Remove(BundleName);
+		if (FContentRequestStats* Stats = StatsMap.Find(BundleName))
+		{
+			ensureAlwaysMsgf(!Stats->bOpen, TEXT("StatsReset - Stat open for %s"), *BundleName.ToString());
+			ensureAlwaysMsgf(
+				Algo::AllOf(Stats->StateStats,
+					[](const TPair<FString, FContentRequestStateStats>& Pair) { return !Pair.Value.bOpen; }),
+				TEXT("StatsReset - StateStat open for %s"), *BundleName.ToString());
+
+			StatsMap.Remove(BundleName);
+		}
 	}
 
 	void FContentRequestStatsMap::StatsBegin(FName BundleName, const TCHAR* State)
 	{
 		FContentRequestStats& Stats = StatsMap.FindOrAdd(BundleName);
-		if (ensureAlwaysMsgf(Stats.bOpen, TEXT("StatsBegin - Stat closed for %s - %s"), *BundleName.ToString(), State) == false)
+		if (false == ensureAlwaysMsgf(Stats.bOpen, TEXT("StatsBegin - Stat closed for %s - %s"), *BundleName.ToString(), State))
 		{
 			Stats = FContentRequestStats();
 			Stats.StartTime = FPlatformTime::Seconds();
 		}
 
 		FContentRequestStateStats& StateStats = Stats.StateStats.FindOrAdd(State);
-		if (ensureAlwaysMsgf(StateStats.bOpen, TEXT("StatsBegin - StateStat closed for %s - %s"), *BundleName.ToString(), State) == false)
+		if (false == ensureAlwaysMsgf(StateStats.bOpen, TEXT("StatsBegin - StateStat closed for %s - %s"), *BundleName.ToString(), State))
 		{
 			StateStats = FContentRequestStateStats();
 		}
@@ -322,14 +338,14 @@ namespace InstallBundleUtil
 	void FContentRequestStatsMap::StatsEnd(FName BundleName, const TCHAR* State, uint64 DataSize /*= 0*/)
 	{
 		FContentRequestStats& Stats = StatsMap.FindOrAdd(BundleName);
-		if (ensureAlwaysMsgf(Stats.bOpen, TEXT("StatsEnd - Stat closed for %s - %s"), *BundleName.ToString(), State) == false)
+		if (false == ensureAlwaysMsgf(Stats.bOpen && Stats.StartTime > 0, TEXT("StatsEnd - Stat closed for %s - %s"), *BundleName.ToString(), State))
 		{
 			Stats = FContentRequestStats();
 			Stats.StartTime = FPlatformTime::Seconds();
 		}
 
 		FContentRequestStateStats& StateStats = Stats.StateStats.FindOrAdd(State);
-		if(ensureAlwaysMsgf(StateStats.bOpen, TEXT("StatsEnd - StateStat closed for %s - %s"), *BundleName.ToString(), State))
+		if (ensureAlwaysMsgf(StateStats.bOpen && StateStats.StartTime > 0, TEXT("StatsEnd - StateStat closed for %s - %s"), *BundleName.ToString(), State))
 		{
 			StateStats.EndTime = FPlatformTime::Seconds();
 			StateStats.DataSize = DataSize;
diff --git a/Engine/Source/Runtime/Online/BuildPatchServices/Private/Installer/DownloadService.cpp b/Engine/Source/Runtime/Online/BuildPatchServices/Private/Installer/DownloadService.cpp
--- a/Engine/Source/Runtime/Online/BuildPatchServices/Private/Installer/DownloadService.cpp
+++ b/Engine/Source/Runtime/Online/BuildPatchServices/Private/Installer/DownloadService.cpp
@@ -490,6 +490,7 @@ namespace BuildPatchServices
 				TSharedRef<IHttpRequest, ESPMode::ThreadSafe> HttpRequest = HttpManager->CreateRequest();
 				HttpRequest->OnRequestProgress64().BindThreadSafeSP(HttpDelegates, &FHttpDelegates::HttpRequestProgress, NewRequest.Key);
 				HttpRequest->OnProcessRequestComplete().BindThreadSafeSP(HttpDelegates, &FHttpDelegates::HttpRequestComplete, MakeDownloadRecord(NewRequest.Key, NewRequest.Value));
+				HttpRequest->SetDelegateThreadPolicy(EHttpRequestDelegateThreadPolicy::CompleteOnHttpThread);
 				HttpRequest->SetURL(NewRequest.Value);
 				HttpRequest->SetVerb(TEXT("GET"));
 				HttpRequest->ProcessRequest();
