Install Bundle cache now accounts for patching overhead
- Overhead is only applied to cache reserved bundles that need to be patched
- Overhead becomes 0 after they are updated
- Add overhead size to cache stats
- Add some missing FStringView versions of BPS manifest functions
- Fix some BPS manifest functions that weren't being mocked

review-30046427 [REDACTED_HANDLE]

diff --git a/Engine/Plugins/Experimental/DefaultInstallBundleManager/Source/Private/DefaultInstallBundleManager.cpp b/Engine/Plugins/Experimental/DefaultInstallBundleManager/Source/Private/DefaultInstallBundleManager.cpp
--- a/Engine/Plugins/Experimental/DefaultInstallBundleManager/Source/Private/DefaultInstallBundleManager.cpp
+++ b/Engine/Plugins/Experimental/DefaultInstallBundleManager/Source/Private/DefaultInstallBundleManager.cpp
@@ -1037,6 +1037,7 @@ void FDefaultInstallBundleManager::CacheEvictionComplete(TSharedRef<IInstallBund
 
 		// Update current size in cache size if eviction succeeded
 		CacheInfo->CurrentInstallSize = 0;
+		CacheInfo->InstallOverheadSize = 0; // Current contract is that overhead only exists for content that may be patched
 		CacheInfo->TimeStamp = FDateTime::MinValue();
 		BundleCache->AddOrUpdateBundle(Source->GetSourceType(), *CacheInfo);
 	}
@@ -1186,13 +1187,12 @@ void FDefaultInstallBundleManager::UpdateBundleSourceComplete(TSharedRef<IInstal
 
 			if (TOptional<FInstallBundleCacheBundleInfo> CacheBundleInfo = BundleCache->GetBundleInfo(Pair.Key, Request->BundleName))
 			{
-				// Cache Analytics
 				if (ResultInfo.Result == EInstallBundleResult::OK)
 				{
+					// Cache Analytics
 					// Only, send on success - if the request were canceled or something, there is no way to tell if the source got far enough
 					// to tell us if we had a cache hit.  We could just go by sizes in the cache, but that will
 					// be wrong in the case that we need to patch but the size stayed the same.
-
 					const bool bCacheHit = !InResultInfo.bContentWasInstalled;
 					if (bCacheHit)
 					{
@@ -1209,6 +1209,9 @@ void FDefaultInstallBundleManager::UpdateBundleSourceComplete(TSharedRef<IInstal
 						InstallBundleManagerAnalytics::FireEvent_BundleCacheMiss(AnalyticsProvider.Get(),
 							BundleInfo.BundleNameString, LexToString(Pair.Key), bWasPatchRequired);
 					}
+
+					// Since the update succeeded, we know that there is no longer any install overhead
+					CacheBundleInfo->InstallOverheadSize = 0;
 				}
 
 				CacheBundleInfo->CurrentInstallSize = ResultInfo.CurrentInstallSize;
@@ -2736,6 +2739,7 @@ EInstallBundleSourceUpdateBundleInfoResult FDefaultInstallBundleManager::OnUpdat
 
 				CacheBundleInfo->BundleName = UpdateInfoPair.Key;
 				CacheBundleInfo->FullInstallSize = SourceBundleInfo.FullInstallSize;
+				CacheBundleInfo->InstallOverheadSize = SourceBundleInfo.InstallOverheadSize;
 				CacheBundleInfo->TimeStamp = SourceBundleInfo.LastAccessTime;
 				CacheBundleInfo->AgeScalar = Source->GetSourceCacheAgeScalar();
 				BundleCache->AddOrUpdateBundle(SourceType, *CacheBundleInfo);
@@ -4572,6 +4576,7 @@ void FDefaultInstallBundleManager::AsyncInit_OnQueryBundleInfoComplete(TSharedRe
 				FInstallBundleCacheBundleInfo CacheBundleInfo;
 				CacheBundleInfo.BundleName = BundleName;
 				CacheBundleInfo.FullInstallSize = SourceBundleInfo.FullInstallSize;
+				CacheBundleInfo.InstallOverheadSize = SourceBundleInfo.InstallOverheadSize;
 				CacheBundleInfo.CurrentInstallSize = SourceBundleInfo.CurrentInstallSize;
 				CacheBundleInfo.TimeStamp = SourceBundleInfo.LastAccessTime;
 				CacheBundleInfo.AgeScalar = BundleSource->GetSourceCacheAgeScalar();
diff --git a/Engine/Source/Runtime/InstallBundleManager/Private/InstallBundleCache.cpp b/Engine/Source/Runtime/InstallBundleManager/Private/InstallBundleCache.cpp
--- a/Engine/Source/Runtime/InstallBundleManager/Private/InstallBundleCache.cpp
+++ b/Engine/Source/Runtime/InstallBundleManager/Private/InstallBundleCache.cpp
@@ -22,6 +22,7 @@ void FInstallBundleCache::AddOrUpdateBundle(EInstallBundleSourceType Source, con
 
 	FPerSourceBundleCacheInfo& Info = PerSourceCacheInfo.FindOrAdd(AddInfo.BundleName).FindOrAdd(Source);
 	Info.FullInstallSize = AddInfo.FullInstallSize;
+	Info.InstallOverheadSize = AddInfo.InstallOverheadSize;
 	Info.CurrentInstallSize = AddInfo.CurrentInstallSize;
 	Info.TimeStamp = AddInfo.TimeStamp;
 	Info.AgeScalar = FMath::Clamp(AddInfo.AgeScalar, 0.1, 1.0);
@@ -55,6 +56,7 @@ TOptional<FInstallBundleCacheBundleInfo> FInstallBundleCache::GetBundleInfo(FNam
 		FInstallBundleCacheBundleInfo& OutInfo = Ret.Emplace();
 		OutInfo.BundleName = BundleName;
 		OutInfo.FullInstallSize = Info->FullInstallSize;
+		OutInfo.InstallOverheadSize = Info->InstallOverheadSize;
 		OutInfo.CurrentInstallSize = Info->CurrentInstallSize;
 		OutInfo.TimeStamp = Info->TimeStamp;
 		OutInfo.AgeScalar = Info->AgeScalar;
@@ -76,6 +78,7 @@ TOptional<FInstallBundleCacheBundleInfo> FInstallBundleCache::GetBundleInfo(EIns
 			FInstallBundleCacheBundleInfo& OutInfo = Ret.Emplace();
 			OutInfo.BundleName = BundleName;
 			OutInfo.FullInstallSize = SourceInfo->FullInstallSize;
+			OutInfo.InstallOverheadSize = SourceInfo->InstallOverheadSize;
 			OutInfo.CurrentInstallSize = SourceInfo->CurrentInstallSize;
 			OutInfo.TimeStamp = SourceInfo->TimeStamp;
 			OutInfo.AgeScalar = SourceInfo->AgeScalar;
@@ -144,16 +147,11 @@ FInstallBundleCacheReserveResult FInstallBundleCache::Reserve(FName BundleName)
 		return Result;
 	}
 
-	if (BundleInfo->FullInstallSize <= BundleInfo->CurrentInstallSize)
-	{
-		BundleInfo->State = ECacheState::Reserved;
-		Result.Result = EInstallBundleCacheReserveResult::Success;
-		return Result;
-	}
-
-	const uint64 SizeNeeded = BundleInfo->FullInstallSize - BundleInfo->CurrentInstallSize;
+	const uint64 SizeNeeded = (BundleInfo->FullInstallSize <= BundleInfo->CurrentInstallSize) ?
+		BundleInfo->InstallOverheadSize :
+		BundleInfo->InstallOverheadSize + (BundleInfo->FullInstallSize - BundleInfo->CurrentInstallSize);
 	const uint64 UsedSize = GetUsedSize();
-	if(TotalSize >= UsedSize + SizeNeeded)
+	if (TotalSize >= UsedSize + SizeNeeded)
 	{
 		BundleInfo->State = ECacheState::Reserved;
 		Result.Result = EInstallBundleCacheReserveResult::Success;
@@ -208,7 +206,7 @@ FInstallBundleCacheReserveResult FInstallBundleCache::Reserve(FName BundleName)
 	}
 
 #if INSTALLBUNDLE_CACHE_DUMP_INFO
-	GetStats(true);
+	GetStats(EInstallBundleCacheDumpToLog::Default, true);
 #endif // INSTALLBUNDLE_CACHE_DUMP_INFO
 
 	return Result;
@@ -244,7 +242,7 @@ FInstallBundleCacheFlushResult FInstallBundleCache::Flush(EInstallBundleSourceTy
 	}
 
 #if INSTALLBUNDLE_CACHE_DUMP_INFO
-	GetStats(true);
+	GetStats(EInstallBundleCacheDumpToLog::Default, true);
 #endif // INSTALLBUNDLE_CACHE_DUMP_INFO
 
 	return Result;
@@ -428,6 +426,7 @@ FInstallBundleCacheStats FInstallBundleCache::GetStats(EInstallBundleCacheDumpTo
 	{ \
 		UE_LOG(LogInstallBundleManager, Verbosity, TEXT("* \tbundle %s"), *BundleName.ToString()) \
 		UE_LOG(LogInstallBundleManager, Verbosity, TEXT("* \t\tfull size: %" UINT64_FMT), Info.FullInstallSize) \
+		UE_LOG(LogInstallBundleManager, Verbosity, TEXT("* \t\toverhead size: %" UINT64_FMT), Info.InstallOverheadSize) \
 		UE_LOG(LogInstallBundleManager, Verbosity, TEXT("* \t\tcurrent size: %" UINT64_FMT), Info.CurrentInstallSize) \
 		UE_LOG(LogInstallBundleManager, Verbosity, TEXT("* \t\treserved: %s"), (Info.State == ECacheState::Reserved) ? TEXT("true") : TEXT("false")) \
 		UE_LOG(LogInstallBundleManager, Verbosity, TEXT("* \t\ttimestamp: %s"), *Info.TimeStamp.ToString()) \
@@ -435,7 +434,7 @@ FInstallBundleCacheStats FInstallBundleCache::GetStats(EInstallBundleCacheDumpTo
 	}
 
 #define INSTALLBUNDLECACHE_CSV_HEADER_LOG(Verbosity) \
-	UE_LOG(LogInstallBundleManager, Verbosity, TEXT("* \tbundle, full size, current size, diff, reserved, timestamp, age scale"))
+	UE_LOG(LogInstallBundleManager, Verbosity, TEXT("* \tbundle, full size, overhead size, current size, diff, reserved, timestamp, age scale"))
 
 #define INSTALLBUNDLECACHE_CSV_LOG(Verbosity) \
 	if (Info.CurrentInstallSize > 0 || Info.State != ECacheState::Released) \
@@ -443,7 +442,7 @@ FInstallBundleCacheStats FInstallBundleCache::GetStats(EInstallBundleCacheDumpTo
 		const TCHAR* Diff = TEXT("="); \
 		if (Info.FullInstallSize > Info.CurrentInstallSize) Diff = TEXT(">"); \
 		else if(Info.FullInstallSize < Info.CurrentInstallSize) Diff = TEXT("<"); \
-		UE_LOG(LogInstallBundleManager, Verbosity, TEXT("* \t%s, %" UINT64_FMT ", %" UINT64_FMT ", %s, %s, %s, %f"), *BundleName.ToString(), Info.FullInstallSize, Info.CurrentInstallSize, Diff, (Info.State == ECacheState::Reserved) ? TEXT("true") : TEXT("false"), *Info.TimeStamp.ToString(), Info.AgeScalar) \
+		UE_LOG(LogInstallBundleManager, Verbosity, TEXT("* \t%s, %" UINT64_FMT ", %" UINT64_FMT ", %" UINT64_FMT ", %s, %s, %s, %f"), *BundleName.ToString(), Info.FullInstallSize, Info.InstallOverheadSize, Info.CurrentInstallSize, Diff, (Info.State == ECacheState::Reserved) ? TEXT("true") : TEXT("false"), *Info.TimeStamp.ToString(), Info.AgeScalar) \
 	}
 
 	auto DumpToLog_Default = [](FName BundleName, const FBundleCacheInfo& Info)
@@ -563,11 +562,14 @@ void FInstallBundleCache::UpdateCacheInfoFromSourceInfo(FName BundleName)
 	FDateTime TimeStamp = FDateTime::MinValue();
 	double AgeScalar = 1.0;
 	uint64 FullInstallSize = 0;
+	uint64 InstallOverheadSize = 0;
 	uint64 CurrentInstallSize = 0;
 	for (const TPair<EInstallBundleSourceType, FPerSourceBundleCacheInfo>& Pair : *SourcesMap)
 	{
 		FullInstallSize += Pair.Value.FullInstallSize;
+		InstallOverheadSize += Pair.Value.InstallOverheadSize;
 		CurrentInstallSize += Pair.Value.CurrentInstallSize;
+
 		if (Pair.Value.CurrentInstallSize > 0)
 		{
 			if (Pair.Value.TimeStamp > TimeStamp)
@@ -586,6 +588,7 @@ void FInstallBundleCache::UpdateCacheInfoFromSourceInfo(FName BundleName)
 	checkf(BundleCacheInfo.FullInstallSize == FullInstallSize || BundleCacheInfo.State != ECacheState::Reserved, TEXT("Bundle %s: FullInstallSize should not be updated while a bundle is Reserved!"), *BundleName.ToString());
 
 	BundleCacheInfo.FullInstallSize = FullInstallSize;
+	BundleCacheInfo.InstallOverheadSize = InstallOverheadSize;
 	BundleCacheInfo.CurrentInstallSize = CurrentInstallSize;
 	BundleCacheInfo.TimeStamp = TimeStamp;
 	BundleCacheInfo.AgeScalar = AgeScalar;
diff --git a/Engine/Source/Runtime/InstallBundleManager/Public/InstallBundleCache.h b/Engine/Source/Runtime/InstallBundleManager/Public/InstallBundleCache.h
--- a/Engine/Source/Runtime/InstallBundleManager/Public/InstallBundleCache.h
+++ b/Engine/Source/Runtime/InstallBundleManager/Public/InstallBundleCache.h
@@ -25,6 +25,7 @@ struct FInstallBundleCacheBundleInfo
 {
 	FName BundleName;
 	uint64 FullInstallSize = 0; // Total disk footprint when this bundle is fully installed
+	uint64 InstallOverheadSize = 0; // Any extra space required to update the bundle
 	uint64 CurrentInstallSize = 0; // Disk footprint of the bundle in it's current state
 	FDateTime TimeStamp = FDateTime::MinValue(); // Last access time for the bundle.  Used for eviction order
 	double AgeScalar = 1.0; // Allow some bundles to "age" slower than others
@@ -108,6 +109,7 @@ private:
 	struct FPerSourceBundleCacheInfo
 	{
 		uint64 FullInstallSize = 0;
+		uint64 InstallOverheadSize = 0;
 		uint64 CurrentInstallSize = 0;
 		FDateTime TimeStamp = FDateTime::MinValue();
 		double AgeScalar = 1.0;
@@ -123,6 +125,7 @@ private:
 	struct FBundleCacheInfo 
 	{
 		uint64 FullInstallSize = 0;
+		uint64 InstallOverheadSize = 0;
 		uint64 CurrentInstallSize = 0;
 		FDateTime TimeStamp = FDateTime::MinValue();
 		double AgeScalar = 1.0;
@@ -142,9 +145,9 @@ private:
 				return 0;
 
 			if (CurrentInstallSize > FullInstallSize)
-				return CurrentInstallSize;
+				return CurrentInstallSize + InstallOverheadSize;
 
-			return FullInstallSize;
+			return FullInstallSize + InstallOverheadSize;
 		}
 	};
 
diff --git a/Engine/Source/Runtime/InstallBundleManager/Public/InstallBundleTypes.h b/Engine/Source/Runtime/InstallBundleManager/Public/InstallBundleTypes.h
--- a/Engine/Source/Runtime/InstallBundleManager/Public/InstallBundleTypes.h
+++ b/Engine/Source/Runtime/InstallBundleManager/Public/InstallBundleTypes.h
@@ -240,6 +240,7 @@ struct FInstallBundleSourceUpdateBundleInfo
 	FString BundleNameString;
 	EInstallBundlePriority Priority = EInstallBundlePriority::Low;
 	uint64 FullInstallSize = 0; // Total disk footprint when this bundle is fully installed
+	uint64 InstallOverheadSize = 0; // Any additional space required to complete installation
 	FDateTime LastAccessTime = FDateTime::MinValue(); // If cached, used to decide eviction order
 	EInstallBundleInstallState BundleContentState = EInstallBundleInstallState::NotInstalled; // Whether this bundle is up to date
 	bool bIsCached = false; // Whether this bundle should be cached if this source has a bundle cache
diff --git a/Engine/Source/Runtime/Online/BuildPatchServices/Private/BuildPatchManifest.cpp b/Engine/Source/Runtime/Online/BuildPatchServices/Private/BuildPatchManifest.cpp
--- a/Engine/Source/Runtime/Online/BuildPatchServices/Private/BuildPatchManifest.cpp
+++ b/Engine/Source/Runtime/Online/BuildPatchServices/Private/BuildPatchManifest.cpp
@@ -888,6 +888,13 @@ TArray<FString> FBuildPatchAppManifest::GetBuildFileList(const TSet<FString>& Ta
 	return Filenames;
 }
 
+TArray<FStringView> FBuildPatchAppManifest::GetBuildFileListView(const TSet<FString>& Tags) const
+{
+	TArray<FStringView> Filenames;
+	GetTaggedFileList(Tags, Filenames);
+	return Filenames;
+}
+
 int64 FBuildPatchAppManifest::GetFileSize(const TArray<FString>& Filenames) const
 {
 	return Algo::Accumulate<int64>(Filenames, 0, [this](int64 Size, const FString& Filename){ return Size + GetFileSize(Filename); });
@@ -898,7 +905,7 @@ int64 FBuildPatchAppManifest::GetFileSize(const TSet<FString>& Filenames) const
 	return Algo::Accumulate<int64>(Filenames, 0, [this](int64 Size, const FString& Filename){ return Size + GetFileSize(Filename); });
 }
 
-int64 FBuildPatchAppManifest::GetFileSize(const FString& Filename) const
+int64 FBuildPatchAppManifest::GetFileSize(FStringView Filename) const
 {
 	const FFileManifest *const *const FileManifest = FileManifestLookup.Find(Filename);
 	if (FileManifest)
@@ -985,6 +992,21 @@ void FBuildPatchAppManifest::GetTaggedFileList(const TSet<FString>& Tags, TArray
 	}
 }
 
+void FBuildPatchAppManifest::GetTaggedFileList(const TSet<FString>& Tags, TArray<FStringView>& TaggedFiles) const
+{
+	for (const FString& Tag : Tags)
+	{
+		const TArray<const FFileManifest*>* const Files = TaggedFilesLookup.Find(Tag);
+		if (Files != nullptr)
+		{
+			for (const FFileManifest* File : *Files)
+			{
+				TaggedFiles.Add(File->Filename);
+			}
+		}
+	}
+}
+
 void FBuildPatchAppManifest::GetTaggedFileList(const TSet<FString>& Tags, TSet<FString>& TaggedFiles) const
 {
 	for (const FString& Tag : Tags)
diff --git a/Engine/Source/Runtime/Online/BuildPatchServices/Private/BuildPatchManifest.h b/Engine/Source/Runtime/Online/BuildPatchServices/Private/BuildPatchManifest.h
--- a/Engine/Source/Runtime/Online/BuildPatchServices/Private/BuildPatchManifest.h
+++ b/Engine/Source/Runtime/Online/BuildPatchServices/Private/BuildPatchManifest.h
@@ -114,7 +114,8 @@ public:
 	virtual TArray<FString> GetBuildFileList() const override;
 	virtual TArray<FStringView> GetBuildFileListView() const override;
 	virtual TArray<FString> GetBuildFileList(const TSet<FString>& Tags) const override;
-	virtual int64 GetFileSize(const FString& Filename) const override;
+	virtual TArray<FStringView> GetBuildFileListView(const TSet<FString>& Tags) const override;
+	virtual int64 GetFileSize(FStringView Filename) const override;
 	virtual int64 GetFileSize(const TArray<FString>& Filenames) const override;
 	virtual int64 GetFileSize(const TSet  <FString>& Filenames) const override;
 	virtual bool GetFileHash(const FString& Filename, FSHAHash& OutHash) const override;
@@ -228,8 +229,9 @@ public:
 	 * @param Tags					The tags for the required file groups.
 	 * @param TaggedFiles	OUT		Receives the tagged files.
 	 */
-	virtual void GetTaggedFileList(const TSet<FString>& Tags, TArray<FString>& TaggedFiles) const;
-	virtual void GetTaggedFileList(const TSet<FString>& Tags, TSet  <FString>& TaggedFiles) const;
+	virtual void GetTaggedFileList(const TSet<FString>& Tags, TArray<FString>&     TaggedFiles) const;
+	virtual void GetTaggedFileList(const TSet<FString>& Tags, TArray<FStringView>& TaggedFiles) const;
+	virtual void GetTaggedFileList(const TSet<FString>& Tags, TSet<FString>&       TaggedFiles) const;
 
 	/**
 	 * Get the list of Guids for all chunks referenced by this manifest
diff --git a/Engine/Source/Runtime/Online/BuildPatchServices/Private/Tests/Mock/Manifest.mock.h b/Engine/Source/Runtime/Online/BuildPatchServices/Private/Tests/Mock/Manifest.mock.h
--- a/Engine/Source/Runtime/Online/BuildPatchServices/Private/Tests/Mock/Manifest.mock.h
+++ b/Engine/Source/Runtime/Online/BuildPatchServices/Private/Tests/Mock/Manifest.mock.h
@@ -2,6 +2,7 @@
 #pragma once
 
 #include "Interfaces/IBuildManifest.h"
+#include "Algo/Copy.h"
 #include "BuildPatchManifest.h"
 #include "Tests/Mock/ManifestField.mock.h"
 
@@ -85,6 +86,13 @@ namespace BuildPatchServices
 			return Filenames;
 		}
 
+		virtual TArray<FStringView> GetBuildFileListView() const override
+		{
+			TArray<FStringView> Filenames;
+			GetFileList(Filenames);
+			return Filenames;
+		}
+
 		virtual TArray<FString> GetBuildFileList(const TSet<FString>& Tags) const override
 		{
 			TArray<FString> Filenames;
@@ -92,6 +100,13 @@ namespace BuildPatchServices
 			return Filenames;
 		}
 
+		virtual TArray<FStringView> GetBuildFileListView(const TSet<FString>& Tags) const override
+		{
+			TArray<FStringView> Filenames;
+			GetTaggedFileList(Tags, Filenames);
+			return Filenames;
+		}
+
 		virtual void GetFileTagList(TSet<FString>& Tags) const override
 		{
 			Tags = FileTagList;
@@ -200,9 +215,9 @@ namespace BuildPatchServices
 			return DataSize;
 		}
 
-		virtual int64 GetFileSize(const FString& Filename) const override
+		virtual int64 GetFileSize(FStringView Filename) const override
 		{
-			return FileNameToFileSize.FindRef(Filename);
+			return FileNameToFileSize.FindRef(FString(Filename));
 		}
 
 		virtual int64 GetFileSize(const TArray<FString>& Filenames) const override
@@ -230,6 +245,16 @@ namespace BuildPatchServices
 			Filenames = BuildFileList;
 		}
 
+		virtual void GetFileList(TArray<FStringView>& Filenames) const override
+		{
+			Algo::Copy(BuildFileList, Filenames);
+		}
+
+		virtual void GetFileList(TSet<FString>& Filenames) const override
+		{
+			Algo::Copy(BuildFileList, Filenames);
+		}
+
 		virtual void GetTaggedFileList(const TSet<FString>& Tags, TSet<FString>& TaggedFiles) const override
 		{
 			TaggedFiles = TaggedFileList;
@@ -240,6 +265,11 @@ namespace BuildPatchServices
 			TaggedFiles = TaggedFileList.Array();
 		}
 
+		virtual void GetTaggedFileList(const TSet<FString>& Tags, TArray<FStringView>& TaggedFiles) const override
+		{
+			Algo::Copy(TaggedFileList, TaggedFiles);
+		}
+
 		virtual void GetDataList(TArray<FGuid>& DataGuids) const override
 		{
 			DataGuids = DataList;
diff --git a/Engine/Source/Runtime/Online/BuildPatchServices/Public/Interfaces/IBuildManifest.h b/Engine/Source/Runtime/Online/BuildPatchServices/Public/Interfaces/IBuildManifest.h
--- a/Engine/Source/Runtime/Online/BuildPatchServices/Public/Interfaces/IBuildManifest.h
+++ b/Engine/Source/Runtime/Online/BuildPatchServices/Public/Interfaces/IBuildManifest.h
@@ -176,12 +176,19 @@ public:
 	 */
 	virtual TArray<FString> GetBuildFileList(const TSet<FString>& Tags) const = 0;
 
+	/**
+	 * Get the list of files in this build which match a tag from a given set. Filenames are all relative to an install directory
+	 * Filenames returned are a view over memory owned by the BuildManifest
+	 * @return an array containing build files
+	 */
+	virtual TArray<FStringView> GetBuildFileListView(const TSet<FString>& Tags) const = 0;
+
 	/**
 	 * Returns the size of a particular file in the build
 	 * @param Filename      The file.
 	 * @return the file size.
 	 */
-	virtual int64 GetFileSize(const FString& Filename) const = 0;
+	virtual int64 GetFileSize(FStringView Filename) const = 0;
 
 	/**
 	 * Returns the total size of all files in the array
