Add a version of StartInstallBundleAsyncIOTask that uses FAutoDeleteInstallBundleTask.  This version does not support passing a game thread callback and does not require the Start/Finish functions to update it.

diff --git a/Engine/Source/Runtime/InstallBundleManager/Private/InstallBundleUtils.cpp b/Engine/Source/Runtime/InstallBundleManager/Private/InstallBundleUtils.cpp
--- a/Engine/Source/Runtime/InstallBundleManager/Private/InstallBundleUtils.cpp
+++ b/Engine/Source/Runtime/InstallBundleManager/Private/InstallBundleUtils.cpp
@@ -273,6 +273,15 @@ namespace InstallBundleUtil
 		return InstallBundleSuppressAnalyticsCounter > 0;
 	}
 
+	void StartInstallBundleAsyncIOTask(TUniqueFunction<void()> WorkFunc)
+	{
+		return StartInstallBundleAsyncIOTask(GIOThreadPool, MoveTemp(WorkFunc));
+	}
+
+	void StartInstallBundleAsyncIOTask(FQueuedThreadPool* ThreadPool, TUniqueFunction<void()> WorkFunc)
+	{
+		(new FAutoDeleteInstallBundleTask(MoveTemp(WorkFunc), nullptr))->StartBackgroundTask(ThreadPool);
+	}
 
 	void StartInstallBundleAsyncIOTask(TArray<TUniquePtr<FInstallBundleTask>>& Tasks, TUniqueFunction<void()> WorkFunc, TUniqueFunction<void()> OnComplete)
 	{
diff --git a/Engine/Source/Runtime/InstallBundleManager/Public/InstallBundleUtils.h b/Engine/Source/Runtime/InstallBundleManager/Public/InstallBundleUtils.h
--- a/Engine/Source/Runtime/InstallBundleManager/Public/InstallBundleUtils.h
+++ b/Engine/Source/Runtime/InstallBundleManager/Public/InstallBundleUtils.h
@@ -199,6 +199,11 @@ namespace InstallBundleUtil
 	};
 
 	using FInstallBundleTask = FAsyncTask<FInstallBundleWork>;
+	using FAutoDeleteInstallBundleTask = FAutoDeleteAsyncTask<FInstallBundleWork>;
+
+	INSTALLBUNDLEMANAGER_API void StartInstallBundleAsyncIOTask(TUniqueFunction<void()> WorkFunc);
+
+	INSTALLBUNDLEMANAGER_API void StartInstallBundleAsyncIOTask(FQueuedThreadPool* ThreadPool, TUniqueFunction<void()> WorkFunc);
 
 	INSTALLBUNDLEMANAGER_API void StartInstallBundleAsyncIOTask(TArray<TUniquePtr<FInstallBundleTask>>& Tasks, TUniqueFunction<void()> WorkFunc, TUniqueFunction<void()> OnComplete);
 
