FInstallBundleCombinedProgressTracker
- Minor cleanup to make logic that determines bAreAllBundlesZeroWeight more clear

diff --git a/Engine/Source/Runtime/InstallBundleManager/Private/BundlePrereqCombinedStatusHelper.cpp b/Engine/Source/Runtime/InstallBundleManager/Private/BundlePrereqCombinedStatusHelper.cpp
--- a/Engine/Source/Runtime/InstallBundleManager/Private/BundlePrereqCombinedStatusHelper.cpp
+++ b/Engine/Source/Runtime/InstallBundleManager/Private/BundlePrereqCombinedStatusHelper.cpp
@@ -121,17 +121,12 @@ void FInstallBundleCombinedProgressTracker::SetBundlesToTrackFromContentState(co
 	for (const TPair<FName, FInstallBundleContentState>& IndividualBundlePair : BundleContentState.IndividualBundleStates)
 	{
 		const FInstallBundleContentState& BundleState = IndividualBundlePair.Value;
-		if (BundleState.Weight <= SMALL_NUMBER)
-		{
-			continue;
-		}
-		else
+		if (BundleState.Weight > SMALL_NUMBER)
 		{
 			bAreAllBundlesZeroWeight = false;
 			break;
 		}
 	}
-		
 
 	bool bBundleNeedsUpdate = false;
 	float TotalWeight = 0.0f;
