Refactored FInstallBundleCombinedContentState to have separate members content size and background download content size (for ChunkDBs)

[RN] FInstallBundleCombinedContentState is now expected to report download and install values separately for regular downloads and, if supported, for Background Downloads (ChunkDB's)

diff --git a/Engine/Plugins/Experimental/DefaultInstallBundleManager/Source/Private/DefaultInstallBundleManager.cpp b/Engine/Plugins/Experimental/DefaultInstallBundleManager/Source/Private/DefaultInstallBundleManager.cpp
--- a/Engine/Plugins/Experimental/DefaultInstallBundleManager/Source/Private/DefaultInstallBundleManager.cpp
+++ b/Engine/Plugins/Experimental/DefaultInstallBundleManager/Source/Private/DefaultInstallBundleManager.cpp
@@ -476,6 +476,8 @@ void FDefaultInstallBundleManager::TickGetContentState()
 
 			FInstallBundleCombinedContentState State;
 			State.FreeSpace = 0;
+			State.BackgroundDownloadContentSize = FInstallBundleContentSize{}; // it's optional for individual sources, but we always have a value
+			
 			for (const TPair<FInstallBundleSourceType, FInstallBundleCombinedContentState>& SourcePair : Request->BundleSourceContentStates)
 			{
 				const FInstallBundleSourceType SourceType = SourcePair.Key;
@@ -510,11 +512,9 @@ void FDefaultInstallBundleManager::TickGetContentState()
 					BundleState->Version.Add(SourceType, StatePair.Value.Version.FindChecked(SourceType));
 				}
 
-				State.DownloadSize += SourceState.DownloadSize;
-				State.InstallSize += SourceState.InstallSize;
-				State.InstallOverheadSize = FMath::Max(State.InstallOverheadSize, SourceState.InstallOverheadSize);
-				State.MaxDiskSpaceRequired += SourceState.MaxDiskSpaceRequired;
-				State.FreeSpace = FMath::Max(State.FreeSpace, SourceState.FreeSpace);
+				State.ContentSize += SourceState.ContentSize;
+				State.BackgroundDownloadContentSize.GetValue() += SourceState.BackgroundDownloadContentSize.Get(SourceState.ContentSize);
+				State.FreeSpace = State.FreeSpace ? FMath::Min(State.FreeSpace, SourceState.FreeSpace) : SourceState.FreeSpace;
 			}
 
 			Request->ExecCallbackIfValid(MoveTemp(State));
diff --git a/Engine/Plugins/Experimental/DefaultInstallBundleManager/Source/Private/InstallBundleManagerUtil.cpp b/Engine/Plugins/Experimental/DefaultInstallBundleManager/Source/Private/InstallBundleManagerUtil.cpp
--- a/Engine/Plugins/Experimental/DefaultInstallBundleManager/Source/Private/InstallBundleManagerUtil.cpp
+++ b/Engine/Plugins/Experimental/DefaultInstallBundleManager/Source/Private/InstallBundleManagerUtil.cpp
@@ -315,8 +315,8 @@ namespace InstallBundleManagerUtil
 	void FPersistentStatContainer::UpdateForContentState(const FInstallBundleCombinedContentState& ContentState, const FString& SessionName)
 	{
 		FSessionAnalyticsData& FoundData = SessionAnalyticsDataMap.FindOrAdd(SessionName);
-		FoundData.bRequiresDownload = ContentState.DownloadSize > 0;
-		FoundData.bRequiresInstall = ContentState.InstallSize > 0;
+		FoundData.bRequiresDownload = ContentState.ContentSize.DownloadSize > 0;
+		FoundData.bRequiresInstall = ContentState.ContentSize.SpaceRequiredForInstall > 0;
 		
 		//Reset so that it can react to the above values being changed
 		FoundData.ResetShouldSendBGAnalytics();
diff --git a/Engine/Source/Runtime/InstallBundleManager/Public/InstallBundleTypes.h b/Engine/Source/Runtime/InstallBundleManager/Public/InstallBundleTypes.h
--- a/Engine/Source/Runtime/InstallBundleManager/Public/InstallBundleTypes.h
+++ b/Engine/Source/Runtime/InstallBundleManager/Public/InstallBundleTypes.h
@@ -124,15 +124,52 @@ struct FInstallBundleContentState
 	TMap<FInstallBundleSourceType, FString> Version;
 };
 
+struct FInstallBundleContentSize
+{
+	uint64 DownloadSize = 0;
+	uint64 SpaceRequiredForInstall = 0;
+	FInstallBundleContentSize operator+(const FInstallBundleContentSize& Other) const
+	{
+		FInstallBundleContentSize Result;
+		Result.DownloadSize = DownloadSize + Other.DownloadSize;
+		Result.SpaceRequiredForInstall = SpaceRequiredForInstall + Other.SpaceRequiredForInstall;
+		return Result;
+	}
+	FInstallBundleContentSize& operator+=(const FInstallBundleContentSize& Other)
+	{
+		*this = *this + Other;
+		return *this;
+	}
+};
+
 struct FInstallBundleCombinedContentState
 {
+	PRAGMA_DISABLE_DEPRECATION_WARNINGS
+	FInstallBundleCombinedContentState() = default;
+	FInstallBundleCombinedContentState(FInstallBundleCombinedContentState&&) = default;
+	FInstallBundleCombinedContentState& operator=(FInstallBundleCombinedContentState&&) = default;
+	FInstallBundleCombinedContentState(const FInstallBundleCombinedContentState&) = default;
+	FInstallBundleCombinedContentState& operator=(const FInstallBundleCombinedContentState&) = default;
+	~FInstallBundleCombinedContentState() = default;
+	PRAGMA_ENABLE_DEPRECATION_WARNINGS
+	
 	TMap<FName, FInstallBundleContentState> IndividualBundleStates;
 	TMap<FInstallBundleSourceType, FString> CurrentVersion;
+	
+#if !UE_BUILD_SHIPPING
+	UE_DEPRECATED(5.6, "Please switch to using ContentSize and BackgroundDownloadContentSize")
 	uint64 DownloadSize = 0;
+	UE_DEPRECATED(5.6, "Please switch to using ContentSize")
 	uint64 InstallSize = 0;
+	UE_DEPRECATED(5.6, "Please switch to using ContentSize")
 	uint64 InstallOverheadSize = 0;
+	UE_DEPRECATED(5.6, "Please switch to using BackgroundDownloadContentSize")
 	uint64 MaxDiskSpaceRequired = 0;
+#endif
+	
 	uint64 FreeSpace = 0;
+	FInstallBundleContentSize ContentSize;
+	TOptional<FInstallBundleContentSize> BackgroundDownloadContentSize;
 
 	INSTALLBUNDLEMANAGER_API bool GetAllBundlesHaveState(EInstallBundleInstallState State, TArrayView<const FName> ExcludedBundles = TArrayView<const FName>()) const;	
 	INSTALLBUNDLEMANAGER_API bool GetAnyBundleHasState(EInstallBundleInstallState State, TArrayView<const FName> ExcludedBundles = TArrayView<const FName>()) const;
