Information about the background download percentage added to the combined status tracker

#rnx

diff --git a/Engine/Plugins/Experimental/DefaultInstallBundleManager/Source/Private/DefaultInstallBundleManager.cpp b/Engine/Plugins/Experimental/DefaultInstallBundleManager/Source/Private/DefaultInstallBundleManager.cpp
--- a/Engine/Plugins/Experimental/DefaultInstallBundleManager/Source/Private/DefaultInstallBundleManager.cpp
+++ b/Engine/Plugins/Experimental/DefaultInstallBundleManager/Source/Private/DefaultInstallBundleManager.cpp
@@ -3916,17 +3916,20 @@ TOptional<FInstallBundleProgress> FDefaultInstallBundleManager::GetBundleProgres
 			const float SourceWeight = BundleSources.FindChecked(Pair.Key)->GetSourceWeight();
 			check(SourceWeight > 0.0f);
 			TotalWeight += SourceWeight;
+			CombinedStatus->BackgroundDownload_Percent += SourceWeight * Pair.Value.BackgroundDownload_Percent;
 			CombinedStatus->Install_Percent += SourceWeight * Pair.Value.Install_Percent;
 		}
 
 		check(TotalWeight >= 0.0f);
 		if (TotalWeight > 0.0f)
 		{
+			CombinedStatus->BackgroundDownload_Percent /= TotalWeight;
 			CombinedStatus->Install_Percent /= TotalWeight;
 		}
 	}
 	else if (State > EContentRequestState::UpdatingBundleSources && Request->Result == EInstallBundleResult::OK)
 	{
+		CombinedStatus->BackgroundDownload_Percent = 1.0f;
 		CombinedStatus->Install_Percent = 1.0f;
 	}
 
diff --git a/Engine/Source/Runtime/InstallBundleManager/Private/BundlePrereqCombinedStatusHelper.cpp b/Engine/Source/Runtime/InstallBundleManager/Private/BundlePrereqCombinedStatusHelper.cpp
--- a/Engine/Source/Runtime/InstallBundleManager/Private/BundlePrereqCombinedStatusHelper.cpp
+++ b/Engine/Source/Runtime/InstallBundleManager/Private/BundlePrereqCombinedStatusHelper.cpp
@@ -170,6 +170,7 @@ void FInstallBundleCombinedProgressTracker::SetBundlesToTrackFromContentState(co
 	// If no bundles to track, we are done
 	if (RequiredBundleNames.Num() == 0)
 	{
+		CurrentCombinedProgress.BackgroundDownloadProgressPercent = 1.0f;
 		CurrentCombinedProgress.ProgressPercent = 1.0f;
 		CurrentCombinedProgress.CombinedStatus = ECombinedBundleStatus::Finished;
 	}
@@ -207,6 +208,7 @@ void FInstallBundleCombinedProgressTracker::UpdateCombinedStatus()
 	if (RequiredBundleNames.Num() == 0)
 		return;
 
+	CurrentCombinedProgress.BackgroundDownloadProgressPercent = GetBackgroundDownloadProgressPercent();
 	CurrentCombinedProgress.ProgressPercent = GetCombinedProgressPercent();
 	
 	EInstallBundleStatus EarliestBundleState = EInstallBundleStatus::Count;
@@ -312,6 +314,24 @@ float FInstallBundleCombinedProgressTracker::GetCombinedProgressPercent() const
 	return FMath::Clamp(AllBundleProgressPercent, 0.f, 1.0f);
 }
 
+float FInstallBundleCombinedProgressTracker::GetBackgroundDownloadProgressPercent() const
+{
+	float BundleDownloadProgressPercent = 0.f;
+	
+	ensureAlwaysMsgf((CachedBundleWeights.Num() >= BundleStatusCache.Num()), TEXT("Missing Cache Entries for BundleWeights!Any missing bundles will have 0 for their progress!"));
+	
+	for (const TPair<FName,FInstallBundleProgress>& BundleStatusPair : BundleStatusCache)
+	{
+		const float* FoundWeight = CachedBundleWeights.Find(BundleStatusPair.Key);
+		if (ensureAlwaysMsgf((nullptr != FoundWeight), TEXT("Found missing entry for BundleWeight! Bundle %s does not have a weight entry!"), *(BundleStatusPair.Key.ToString())))
+		{
+			BundleDownloadProgressPercent += ((*FoundWeight) * BundleStatusPair.Value.BackgroundDownload_Percent);
+		}
+	}
+	
+	return FMath::Clamp(BundleDownloadProgressPercent, 0.f, 1.0f);
+}
+
 bool FInstallBundleCombinedProgressTracker::Tick(float dt)
 {
 	QUICK_SCOPE_CYCLE_COUNTER(STAT_FBundlePrereqCombinedStatusHelper_Tick);
diff --git a/Engine/Source/Runtime/InstallBundleManager/Public/BundlePrereqCombinedStatusHelper.h b/Engine/Source/Runtime/InstallBundleManager/Public/BundlePrereqCombinedStatusHelper.h
--- a/Engine/Source/Runtime/InstallBundleManager/Public/BundlePrereqCombinedStatusHelper.h
+++ b/Engine/Source/Runtime/InstallBundleManager/Public/BundlePrereqCombinedStatusHelper.h
@@ -40,6 +40,7 @@ public:
 	//provide all our needed combined status information in 1 struct
 	struct FCombinedProgress
 	{
+		float BackgroundDownloadProgressPercent = 0.0f;
 		float ProgressPercent = 0.0f;
 		ECombinedBundleStatus CombinedStatus = ECombinedBundleStatus::Unknown;
 		EInstallBundlePauseFlags CombinedPauseFlags = EInstallBundlePauseFlags::None;
@@ -81,6 +82,7 @@ private:
 	INSTALLBUNDLEMANAGER_API void OnBundleInstallPauseChanged(FInstallBundlePauseInfo PauseInfo);
 	
 	INSTALLBUNDLEMANAGER_API float GetCombinedProgressPercent() const;
+	INSTALLBUNDLEMANAGER_API float GetBackgroundDownloadProgressPercent() const;
 	
 private:
 	//All bundles we need including pre-reqs
diff --git a/Engine/Source/Runtime/InstallBundleManager/Public/InstallBundleManagerInterface.h b/Engine/Source/Runtime/InstallBundleManager/Public/InstallBundleManagerInterface.h
--- a/Engine/Source/Runtime/InstallBundleManager/Public/InstallBundleManagerInterface.h
+++ b/Engine/Source/Runtime/InstallBundleManager/Public/InstallBundleManagerInterface.h
@@ -38,6 +38,7 @@ struct FInstallBundleProgress
 
 	EInstallBundlePauseFlags PauseFlags = EInstallBundlePauseFlags::None;
 
+	float BackgroundDownload_Percent = 0;
 	float Install_Percent = 0;
 
 	float Finishing_Percent = 0;
diff --git a/Engine/Source/Runtime/InstallBundleManager/Public/InstallBundleTypes.h b/Engine/Source/Runtime/InstallBundleManager/Public/InstallBundleTypes.h
--- a/Engine/Source/Runtime/InstallBundleManager/Public/InstallBundleTypes.h
+++ b/Engine/Source/Runtime/InstallBundleManager/Public/InstallBundleTypes.h
@@ -395,6 +395,7 @@ struct FInstallBundleSourceProgress
 {
 	FName BundleName;
 
+	float BackgroundDownload_Percent = 0;
 	float Install_Percent = 0;
 };
 
