BundlePrereqCombinedStatusHelper: added tracking for installed bundles and total bundles count

#rnx

diff --git a/Engine/Source/Runtime/InstallBundleManager/Private/BundlePrereqCombinedStatusHelper.cpp b/Engine/Source/Runtime/InstallBundleManager/Private/BundlePrereqCombinedStatusHelper.cpp
--- a/Engine/Source/Runtime/InstallBundleManager/Private/BundlePrereqCombinedStatusHelper.cpp
+++ b/Engine/Source/Runtime/InstallBundleManager/Private/BundlePrereqCombinedStatusHelper.cpp
@@ -208,8 +208,7 @@ void FInstallBundleCombinedProgressTracker::UpdateCombinedStatus()
 	if (RequiredBundleNames.Num() == 0)
 		return;
 
-	CurrentCombinedProgress.BackgroundDownloadProgressPercent = GetBackgroundDownloadProgressPercent();
-	CurrentCombinedProgress.ProgressPercent = GetCombinedProgressPercent();
+	UpdateCombinedProgress();
 	
 	EInstallBundleStatus EarliestBundleState = EInstallBundleStatus::Count;
 	EInstallBundlePauseFlags CombinedPauseFlags = EInstallBundlePauseFlags::None;
@@ -296,9 +295,12 @@ void FInstallBundleCombinedProgressTracker::UpdateCombinedStatus()
 	}
 }
 
-float FInstallBundleCombinedProgressTracker::GetCombinedProgressPercent() const
+void FInstallBundleCombinedProgressTracker::UpdateCombinedProgress()
 {
 	float AllBundleProgressPercent = 0.f;
+	float AllBundleDownloadProgressPercent = 0.f;
+	int32 InstalledBundlesCount = 0;
+	int32 TotalBundlesCount = 0;
 	
 	ensureAlwaysMsgf((CachedBundleWeights.Num() >= BundleStatusCache.Num()), TEXT("Missing Cache Entries for BundleWeights!Any missing bundles will have 0 for their progress!"));
 	
@@ -308,28 +310,19 @@ float FInstallBundleCombinedProgressTracker::GetCombinedProgressPercent() const
 		if (ensureAlwaysMsgf((nullptr != FoundWeight), TEXT("Found missing entry for BundleWeight! Bundle %s does not have a weight entry!"), *(BundleStatusPair.Key.ToString())))
 		{
 			AllBundleProgressPercent += ((*FoundWeight) * BundleStatusPair.Value.Install_Percent);
+			AllBundleDownloadProgressPercent += ((*FoundWeight) * BundleStatusPair.Value.BackgroundDownload_Percent);
+			if (FMath::IsNearlyEqual(BundleStatusPair.Value.Install_Percent, 1.0f))
+			{
+				InstalledBundlesCount++;
+			}
+			TotalBundlesCount++;
 		}
 	}
-	
-	return FMath::Clamp(AllBundleProgressPercent, 0.f, 1.0f);
-}
 
-float FInstallBundleCombinedProgressTracker::GetBackgroundDownloadProgressPercent() const
-{
-	float BundleDownloadProgressPercent = 0.f;
-	
-	ensureAlwaysMsgf((CachedBundleWeights.Num() >= BundleStatusCache.Num()), TEXT("Missing Cache Entries for BundleWeights!Any missing bundles will have 0 for their progress!"));
-	
-	for (const TPair<FName,FInstallBundleProgress>& BundleStatusPair : BundleStatusCache)
-	{
-		const float* FoundWeight = CachedBundleWeights.Find(BundleStatusPair.Key);
-		if (ensureAlwaysMsgf((nullptr != FoundWeight), TEXT("Found missing entry for BundleWeight! Bundle %s does not have a weight entry!"), *(BundleStatusPair.Key.ToString())))
-		{
-			BundleDownloadProgressPercent += ((*FoundWeight) * BundleStatusPair.Value.BackgroundDownload_Percent);
-		}
-	}
-	
-	return FMath::Clamp(BundleDownloadProgressPercent, 0.f, 1.0f);
+	CurrentCombinedProgress.ProgressPercent = FMath::Clamp(AllBundleProgressPercent, 0.f, 1.0f);
+	CurrentCombinedProgress.BackgroundDownloadProgressPercent = FMath::Clamp(AllBundleDownloadProgressPercent, 0.f, 1.0f);
+	CurrentCombinedProgress.InstalledBundlesCount = InstalledBundlesCount;
+	CurrentCombinedProgress.TotalBundlesCount = TotalBundlesCount;
 }
 
 bool FInstallBundleCombinedProgressTracker::Tick(float dt)
diff --git a/Engine/Source/Runtime/InstallBundleManager/Public/BundlePrereqCombinedStatusHelper.h b/Engine/Source/Runtime/InstallBundleManager/Public/BundlePrereqCombinedStatusHelper.h
--- a/Engine/Source/Runtime/InstallBundleManager/Public/BundlePrereqCombinedStatusHelper.h
+++ b/Engine/Source/Runtime/InstallBundleManager/Public/BundlePrereqCombinedStatusHelper.h
@@ -42,6 +42,8 @@ public:
 	{
 		float BackgroundDownloadProgressPercent = 0.0f;
 		float ProgressPercent = 0.0f;
+		int32 InstalledBundlesCount = 0;
+		int32 TotalBundlesCount = 0;
 		ECombinedBundleStatus CombinedStatus = ECombinedBundleStatus::Unknown;
 		EInstallBundlePauseFlags CombinedPauseFlags = EInstallBundlePauseFlags::None;
 		bool bIsPaused = false;
@@ -81,8 +83,7 @@ private:
 	INSTALLBUNDLEMANAGER_API void OnBundleInstallComplete(FInstallBundleRequestResultInfo CompletedBundleInfo);
 	INSTALLBUNDLEMANAGER_API void OnBundleInstallPauseChanged(FInstallBundlePauseInfo PauseInfo);
 	
-	INSTALLBUNDLEMANAGER_API float GetCombinedProgressPercent() const;
-	INSTALLBUNDLEMANAGER_API float GetBackgroundDownloadProgressPercent() const;
+	INSTALLBUNDLEMANAGER_API void UpdateCombinedProgress();
 	
 private:
 	//All bundles we need including pre-reqs
