Add support for BundleSize which will provide the full bundle size if the bundle is installed or not

diff --git a/Engine/Source/Runtime/InstallBundleManager/Public/InstallBundleTypes.h b/Engine/Source/Runtime/InstallBundleManager/Public/InstallBundleTypes.h
--- a/Engine/Source/Runtime/InstallBundleManager/Public/InstallBundleTypes.h
+++ b/Engine/Source/Runtime/InstallBundleManager/Public/InstallBundleTypes.h
@@ -126,13 +126,19 @@ struct FInstallBundleContentState
 
 struct FInstallBundleContentSize
 {
+	// Size remaining to download
 	uint64 DownloadSize = 0;
+	// Size needed to install the remaining download.
 	uint64 SpaceRequiredForInstall = 0;
+	// size of bundle currently on disk.
+	uint64 CurrentSizeOnDisk = 0;
+
 	FInstallBundleContentSize operator+(const FInstallBundleContentSize& Other) const
 	{
 		FInstallBundleContentSize Result;
 		Result.DownloadSize = DownloadSize + Other.DownloadSize;
 		Result.SpaceRequiredForInstall = SpaceRequiredForInstall + Other.SpaceRequiredForInstall;
+		Result.CurrentSizeOnDisk = CurrentSizeOnDisk + Other.CurrentSizeOnDisk;
 		return Result;
 	}
 	FInstallBundleContentSize& operator+=(const FInstallBundleContentSize& Other)
