Add file URI protocol to InstallBundleUtil::SplitHostUrl.  Some tests use file.

diff --git a/Engine/Source/Runtime/InstallBundleManager/Private/InstallBundleUtils.cpp b/Engine/Source/Runtime/InstallBundleManager/Private/InstallBundleUtils.cpp
--- a/Engine/Source/Runtime/InstallBundleManager/Private/InstallBundleUtils.cpp
+++ b/Engine/Source/Runtime/InstallBundleManager/Private/InstallBundleUtils.cpp
@@ -40,6 +40,7 @@ namespace InstallBundleUtil
 
 		const FStringView ProtoHttp = TEXTVIEW("http://");
 		const FStringView ProtoHttps = TEXTVIEW("https://");
+		const FStringView ProtoFile = TEXTVIEW("file://");
 
 		FStringView Proto;
 		if (Url.StartsWith(ProtoHttp))
@@ -50,6 +51,10 @@ namespace InstallBundleUtil
 		{
 			Proto = ProtoHttps;
 		}
+		else if (Url.StartsWith(ProtoFile))
+		{
+			Proto = ProtoFile;
+		}
 		else
 		{
 			return false;
