Install combined progress: InstallOnly_Percent added to track the progress of installs without including the download progress

#rnx

diff --git a/Engine/Plugins/Experimental/DefaultInstallBundleManager/Source/Private/DefaultInstallBundleManager.cpp b/Engine/Plugins/Experimental/DefaultInstallBundleManager/Source/Private/DefaultInstallBundleManager.cpp
--- a/Engine/Plugins/Experimental/DefaultInstallBundleManager/Source/Private/DefaultInstallBundleManager.cpp
+++ b/Engine/Plugins/Experimental/DefaultInstallBundleManager/Source/Private/DefaultInstallBundleManager.cpp
@@ -4106,6 +4106,7 @@ TOptional<FInstallBundleProgress> FDefaultInstallBundleManager::GetBundleProgres
 			check(SourceWeight > 0.0f);
 			TotalWeight += SourceWeight;
 			CombinedStatus->BackgroundDownload_Percent += SourceWeight * Pair.Value.BackgroundDownload_Percent;
+			CombinedStatus->InstallOnly_Percent += SourceWeight * (Pair.Value.InstallOnly_Percent >= 0 ? Pair.Value.InstallOnly_Percent : Pair.Value.Install_Percent);
 			CombinedStatus->Install_Percent += SourceWeight * Pair.Value.Install_Percent;
 		}
 
@@ -4113,12 +4114,14 @@ TOptional<FInstallBundleProgress> FDefaultInstallBundleManager::GetBundleProgres
 		if (TotalWeight > 0.0f)
 		{
 			CombinedStatus->BackgroundDownload_Percent /= TotalWeight;
+			CombinedStatus->InstallOnly_Percent /= TotalWeight;
 			CombinedStatus->Install_Percent /= TotalWeight;
 		}
 	}
 	else if (State > EContentRequestState::UpdatingBundleSources && Request->Result == EInstallBundleResult::OK)
 	{
 		CombinedStatus->BackgroundDownload_Percent = 1.0f;
+		CombinedStatus->InstallOnly_Percent = 1.0f;
 		CombinedStatus->Install_Percent = 1.0f;
 	}
 
diff --git a/Engine/Source/Runtime/InstallBundleManager/Private/BundlePrereqCombinedStatusHelper.cpp b/Engine/Source/Runtime/InstallBundleManager/Private/BundlePrereqCombinedStatusHelper.cpp
--- a/Engine/Source/Runtime/InstallBundleManager/Private/BundlePrereqCombinedStatusHelper.cpp
+++ b/Engine/Source/Runtime/InstallBundleManager/Private/BundlePrereqCombinedStatusHelper.cpp
@@ -171,6 +171,7 @@ void FInstallBundleCombinedProgressTracker::SetBundlesToTrackFromContentState(co
 	if (RequiredBundleNames.Num() == 0)
 	{
 		CurrentCombinedProgress.BackgroundDownloadProgressPercent = 1.0f;
+		CurrentCombinedProgress.InstallOnlyProgressPercent = 1.0f;
 		CurrentCombinedProgress.ProgressPercent = 1.0f;
 		CurrentCombinedProgress.CombinedStatus = ECombinedBundleStatus::Finished;
 	}
@@ -299,6 +300,7 @@ void FInstallBundleCombinedProgressTracker::UpdateCombinedProgress()
 {
 	float AllBundleProgressPercent = 0.f;
 	float AllBundleDownloadProgressPercent = 0.f;
+	float AllBundleInstallOnlyProgressPercent = 0.f;
 	int32 InstalledBundlesCount = 0;
 	int32 TotalBundlesCount = 0;
 	
@@ -311,6 +313,7 @@ void FInstallBundleCombinedProgressTracker::UpdateCombinedProgress()
 		{
 			AllBundleProgressPercent += ((*FoundWeight) * BundleStatusPair.Value.Install_Percent);
 			AllBundleDownloadProgressPercent += ((*FoundWeight) * BundleStatusPair.Value.BackgroundDownload_Percent);
+			AllBundleInstallOnlyProgressPercent += ((*FoundWeight) * BundleStatusPair.Value.InstallOnly_Percent);
 			if (FMath::IsNearlyEqual(BundleStatusPair.Value.Install_Percent, 1.0f))
 			{
 				InstalledBundlesCount++;
@@ -321,6 +324,7 @@ void FInstallBundleCombinedProgressTracker::UpdateCombinedProgress()
 
 	CurrentCombinedProgress.ProgressPercent = FMath::Clamp(AllBundleProgressPercent, 0.f, 1.0f);
 	CurrentCombinedProgress.BackgroundDownloadProgressPercent = FMath::Clamp(AllBundleDownloadProgressPercent, 0.f, 1.0f);
+	CurrentCombinedProgress.InstallOnlyProgressPercent = FMath::Clamp(AllBundleInstallOnlyProgressPercent, 0.f, 1.0f);
 	CurrentCombinedProgress.InstalledBundlesCount = InstalledBundlesCount;
 	CurrentCombinedProgress.TotalBundlesCount = TotalBundlesCount;
 }
diff --git a/Engine/Source/Runtime/InstallBundleManager/Public/BundlePrereqCombinedStatusHelper.h b/Engine/Source/Runtime/InstallBundleManager/Public/BundlePrereqCombinedStatusHelper.h
--- a/Engine/Source/Runtime/InstallBundleManager/Public/BundlePrereqCombinedStatusHelper.h
+++ b/Engine/Source/Runtime/InstallBundleManager/Public/BundlePrereqCombinedStatusHelper.h
@@ -41,7 +41,8 @@ public:
 	struct FCombinedProgress
 	{
 		float BackgroundDownloadProgressPercent = 0.0f;
-		float ProgressPercent = 0.0f;
+		float InstallOnlyProgressPercent = 0.0f;
+		float ProgressPercent = 0.0f; // Download and Install progress combined
 		int32 InstalledBundlesCount = 0;
 		int32 TotalBundlesCount = 0;
 		ECombinedBundleStatus CombinedStatus = ECombinedBundleStatus::Unknown;
diff --git a/Engine/Source/Runtime/InstallBundleManager/Public/InstallBundleManagerInterface.h b/Engine/Source/Runtime/InstallBundleManager/Public/InstallBundleManagerInterface.h
--- a/Engine/Source/Runtime/InstallBundleManager/Public/InstallBundleManagerInterface.h
+++ b/Engine/Source/Runtime/InstallBundleManager/Public/InstallBundleManagerInterface.h
@@ -39,7 +39,8 @@ struct FInstallBundleProgress
 	EInstallBundlePauseFlags PauseFlags = EInstallBundlePauseFlags::None;
 
 	float BackgroundDownload_Percent = 0;
-	float Install_Percent = 0;
+	float InstallOnly_Percent = 0;
+	float Install_Percent = 0; // Download and Install progress combined
 
 	float Finishing_Percent = 0;
 };
diff --git a/Engine/Source/Runtime/InstallBundleManager/Public/InstallBundleTypes.h b/Engine/Source/Runtime/InstallBundleManager/Public/InstallBundleTypes.h
--- a/Engine/Source/Runtime/InstallBundleManager/Public/InstallBundleTypes.h
+++ b/Engine/Source/Runtime/InstallBundleManager/Public/InstallBundleTypes.h
@@ -417,7 +417,8 @@ struct FInstallBundleSourceProgress
 	FName BundleName;
 
 	float BackgroundDownload_Percent = 0;
-	float Install_Percent = 0;
+	float InstallOnly_Percent = -1; // -1 means the value is not valid and Install_Percent should be used instead
+	float Install_Percent = 0; // Download and Install progress combined
 };
 
 struct FInstallBundleSourcePauseInfo
