Bundle installers stats array functionality added

#rnx

diff --git a/Engine/Plugins/Experimental/DefaultInstallBundleManager/Source/Private/DefaultInstallBundleManager.cpp b/Engine/Plugins/Experimental/DefaultInstallBundleManager/Source/Private/DefaultInstallBundleManager.cpp
--- a/Engine/Plugins/Experimental/DefaultInstallBundleManager/Source/Private/DefaultInstallBundleManager.cpp
+++ b/Engine/Plugins/Experimental/DefaultInstallBundleManager/Source/Private/DefaultInstallBundleManager.cpp
@@ -4128,6 +4128,12 @@ TOptional<FInstallBundleProgress> FDefaultInstallBundleManager::GetBundleProgres
 
 	EContentRequestState State = Request->Steps[Request->iStep];
 
+	// Combine bundle installer stats
+	for (const TPair<FInstallBundleSourceType, TSharedPtr<IInstallBundleSource>>& Pair : BundleSources)
+	{
+		CombinedStatus->Stats.Append(Pair.Value->GetBundleUpdateStats(BundleName));
+	}
+
 	if (State == EContentRequestState::UpdatingBundleSources && Request->Result == EInstallBundleResult::OK)
 	{
 		// Update cached status
diff --git a/Engine/Source/Runtime/InstallBundleManager/Private/BundlePrereqCombinedStatusHelper.cpp b/Engine/Source/Runtime/InstallBundleManager/Private/BundlePrereqCombinedStatusHelper.cpp
--- a/Engine/Source/Runtime/InstallBundleManager/Private/BundlePrereqCombinedStatusHelper.cpp
+++ b/Engine/Source/Runtime/InstallBundleManager/Private/BundlePrereqCombinedStatusHelper.cpp
@@ -303,6 +303,7 @@ void FInstallBundleCombinedProgressTracker::UpdateCombinedProgress()
 	float AllBundleInstallOnlyProgressPercent = 0.f;
 	int32 InstalledBundlesCount = 0;
 	int32 TotalBundlesCount = 0;
+	CurrentCombinedProgress.Stats.Empty();
 	
 	ensureAlwaysMsgf((CachedBundleWeights.Num() >= BundleStatusCache.Num()), TEXT("Missing Cache Entries for BundleWeights!Any missing bundles will have 0 for their progress!"));
 	
@@ -319,6 +320,8 @@ void FInstallBundleCombinedProgressTracker::UpdateCombinedProgress()
 				InstalledBundlesCount++;
 			}
 			TotalBundlesCount++;
+
+			CurrentCombinedProgress.Stats.Append(BundleStatusPair.Value.Stats);
 		}
 	}
 
diff --git a/Engine/Source/Runtime/InstallBundleManager/Public/BundlePrereqCombinedStatusHelper.h b/Engine/Source/Runtime/InstallBundleManager/Public/BundlePrereqCombinedStatusHelper.h
--- a/Engine/Source/Runtime/InstallBundleManager/Public/BundlePrereqCombinedStatusHelper.h
+++ b/Engine/Source/Runtime/InstallBundleManager/Public/BundlePrereqCombinedStatusHelper.h
@@ -45,6 +45,7 @@ public:
 		float ProgressPercent = 0.0f; // Download and Install progress combined
 		int32 InstalledBundlesCount = 0;
 		int32 TotalBundlesCount = 0;
+		TArray<FBuildInstallerStat> Stats; // Used for additional information about the installs
 		ECombinedBundleStatus CombinedStatus = ECombinedBundleStatus::Unknown;
 		EInstallBundlePauseFlags CombinedPauseFlags = EInstallBundlePauseFlags::None;
 		bool bIsPaused = false;
diff --git a/Engine/Source/Runtime/InstallBundleManager/Public/InstallBundleManagerInterface.h b/Engine/Source/Runtime/InstallBundleManager/Public/InstallBundleManagerInterface.h
--- a/Engine/Source/Runtime/InstallBundleManager/Public/InstallBundleManagerInterface.h
+++ b/Engine/Source/Runtime/InstallBundleManager/Public/InstallBundleManagerInterface.h
@@ -38,6 +38,8 @@ struct FInstallBundleProgress
 	float BackgroundDownload_Percent = 0;
 	float InstallOnly_Percent = 0;
 	float Install_Percent = 0; // Download and Install progress combined
+	
+	TArray<FBuildInstallerStat> Stats; // Used for additional information about the install
 
 	float Finishing_Percent = 0;
 };
diff --git a/Engine/Source/Runtime/InstallBundleManager/Public/InstallBundleSourceInterface.h b/Engine/Source/Runtime/InstallBundleManager/Public/InstallBundleSourceInterface.h
--- a/Engine/Source/Runtime/InstallBundleManager/Public/InstallBundleSourceInterface.h
+++ b/Engine/Source/Runtime/InstallBundleManager/Public/InstallBundleSourceInterface.h
@@ -132,6 +132,9 @@ public:
 
 	// Derived classes should implement this if their content install will take a significant amount of time
 	virtual TOptional<FInstallBundleSourceProgress> GetBundleProgress(FName BundleName) const { return TOptional<FInstallBundleSourceProgress>(); }
+	
+	// Override to return information about a bundle installer
+	virtual TArray<FBuildInstallerStat>& GetBundleUpdateStats(FName BundleName) { return EmptyStats; }
 
 	virtual void CheckForContentPatch(FInstallBundleSourceContentPatchResultDelegate Callback) { Callback.Execute(AsShared(), false); }
 
@@ -141,4 +144,7 @@ public:
 #if !UE_BUILD_SHIPPING
 	virtual void GetDebugText(TArray<FString>& Output) {}
 #endif
+
+private:
+	TArray<FBuildInstallerStat> EmptyStats;
 };
diff --git a/Engine/Source/Runtime/InstallBundleManager/Public/InstallBundleTypes.h b/Engine/Source/Runtime/InstallBundleManager/Public/InstallBundleTypes.h
--- a/Engine/Source/Runtime/InstallBundleManager/Public/InstallBundleTypes.h
+++ b/Engine/Source/Runtime/InstallBundleManager/Public/InstallBundleTypes.h
@@ -14,6 +14,7 @@
 #include "Misc/DateTime.h"
 #include "Misc/EnumClassFlags.h"
 #include "Misc/EnumRange.h"
+#include "Misc/TVariant.h"
 #include "UObject/NameTypes.h"
 
 #if !defined(WITH_PLATFORM_INSTALL_BUNDLE_SOURCE)
@@ -412,6 +413,14 @@ struct FInstallBundleSourceReleaseContentResultInfo
 	bool bContentWasRemoved = false;
 };
 
+// Useful to store any kind of information about the build installer
+struct FBuildInstallerStat
+{
+	FName BundleName; // To known which Bundle this stat was for even after all the stats are aggregated
+	FName StatName;
+	TVariant<bool, int32, int64, float, double, FString> StatValue;
+};
+
 struct FInstallBundleSourceProgress
 {
 	FName BundleName;
@@ -419,6 +428,8 @@ struct FInstallBundleSourceProgress
 	float BackgroundDownload_Percent = 0;
 	float InstallOnly_Percent = -1; // -1 means the value is not valid and Install_Percent should be used instead
 	float Install_Percent = 0; // Download and Install progress combined
+	
+	TArray<FBuildInstallerStat> Stats; // Used for additional information about the install
 };
 
 struct FInstallBundleSourcePauseInfo
