[Partial Backout] - CL40591438 - Kept the change in IOSPlatformMisc.cpp
#rnx
Original CL Desc
-----------------------------------------------------------------
Fix nw_path_t being released unexpectedly. Add nw_retain/nw_release

#rnx

diff --git a/Engine/Source/Runtime/InstallBundleManager/Private/InstallBundleUtils.cpp b/Engine/Source/Runtime/InstallBundleManager/Private/InstallBundleUtils.cpp
--- a/Engine/Source/Runtime/InstallBundleManager/Private/InstallBundleUtils.cpp
+++ b/Engine/Source/Runtime/InstallBundleManager/Private/InstallBundleUtils.cpp
@@ -31,8 +31,7 @@ namespace InstallBundleUtil
 	bool HasInternetConnection(ENetworkConnectionType ConnectionType)
 	{
 		return ConnectionType != ENetworkConnectionType::AirplaneMode
-			&& ConnectionType != ENetworkConnectionType::None
-			&& ConnectionType != ENetworkConnectionType::Unknown;
+			&& ConnectionType != ENetworkConnectionType::None;
 	}
 
 	bool SplitHostUrl(const FStringView Url, FStringView& OutHost, FStringView& OutRemainder)
