Install Bundle Cache Stats
- Allow dumping stats for individual bundles to FInstallBundleCacheStats
- Deprecate EInstallBundleCacheDumpToLog and replace with EInstallBundleCacheStatsFlags with DumpToResults flag

diff --git a/Engine/Plugins/Experimental/DefaultInstallBundleManager/Source/Private/DefaultInstallBundleManager.cpp b/Engine/Plugins/Experimental/DefaultInstallBundleManager/Source/Private/DefaultInstallBundleManager.cpp
--- a/Engine/Plugins/Experimental/DefaultInstallBundleManager/Source/Private/DefaultInstallBundleManager.cpp
+++ b/Engine/Plugins/Experimental/DefaultInstallBundleManager/Source/Private/DefaultInstallBundleManager.cpp
@@ -115,11 +115,12 @@ static FAutoConsoleCommand InstallBundleCacheStatsCommand(
 		}
 		else if(Args.Num() > 0 && Args[0] == TEXT("CSV"))
 		{
-			InstallBundleManager->GetCacheStats(EInstallBundleCacheDumpToLog::CSV);
+			InstallBundleManager->GetCacheStats(
+				EInstallBundleCacheStatsFlags::DumpToLog | EInstallBundleCacheStatsFlags::CSVFormat);
 		}
 		else
 		{
-			InstallBundleManager->GetCacheStats(EInstallBundleCacheDumpToLog::Default);
+			InstallBundleManager->GetCacheStats(EInstallBundleCacheStatsFlags::DumpToLog);
 		}
 	}),
 	ECVF_Cheat);
@@ -986,7 +987,7 @@ void FDefaultInstallBundleManager::TryReserveCache(FContentRequestRef Request)
 			else if (Pair.Value == EInstallBundleCacheReserveResult::Fail_CacheFull)
 			{
 				// Dump useful info
-				GetCacheStats(FInstallBundleSourceOrCache(Pair.Key), EInstallBundleCacheDumpToLog::Default, Request->GetLogVerbosityOverride());
+				GetCacheStats(FInstallBundleSourceOrCache(Pair.Key), EInstallBundleCacheStatsFlags::DumpToLog, Request->GetLogVerbosityOverride());
 
 				TSharedRef<FInstallBundleCache> Cache = BundleCaches.FindChecked(Pair.Key);
 				if (TOptional<FInstallBundleCacheBundleInfo> CacheBundleInfo = Cache->GetBundleInfo(Request->BundleName))
@@ -3869,7 +3870,7 @@ EInstallBundleResult FDefaultInstallBundleManager::FlushCache(FInstallBundleSour
 	return EInstallBundleResult::OK;
 }
 
-TArray<FInstallBundleCacheStats> FDefaultInstallBundleManager::GetCacheStats(EInstallBundleCacheDumpToLog DumpToLog /*= EInstallBundleCacheDumpToLog::None*/, ELogVerbosity::Type LogVerbosityOverride /*= ELogVerbosity::NoLogging*/)
+TArray<FInstallBundleCacheStats> FDefaultInstallBundleManager::GetCacheStats(EInstallBundleCacheStatsFlags Flags /*= EInstallBundleCacheStatsFlags::None*/, ELogVerbosity::Type LogVerbosityOverride /*= ELogVerbosity::NoLogging*/)
 {
 	TArray<FInstallBundleCacheStats> Stats;
 
@@ -3877,14 +3878,14 @@ TArray<FInstallBundleCacheStats> FDefaultInstallBundleManager::GetCacheStats(EIn
 	{
 		for (const TPair<FName, TSharedRef<FInstallBundleCache>>& Pair : BundleCaches)
 		{
-			Stats.Add(Pair.Value->GetStats(DumpToLog, LogVerbosityOverride < ELogVerbosity::Verbose));
+			Stats.Add(Pair.Value->GetStats(Flags, LogVerbosityOverride < ELogVerbosity::Verbose));
 		}
 	}
 
 	return Stats;
 }
 
-TOptional<FInstallBundleCacheStats> FDefaultInstallBundleManager::GetCacheStats(FInstallBundleSourceOrCache SourceOrCache, EInstallBundleCacheDumpToLog DumpToLog /*= EInstallBundleCacheDumpToLog::None*/, ELogVerbosity::Type LogVerbosityOverride /*= ELogVerbosity::NoLogging*/)
+TOptional<FInstallBundleCacheStats> FDefaultInstallBundleManager::GetCacheStats(FInstallBundleSourceOrCache SourceOrCache, EInstallBundleCacheStatsFlags Flags /*= EInstallBundleCacheStatsFlags::None*/, ELogVerbosity::Type LogVerbosityOverride /*= ELogVerbosity::NoLogging*/)
 {
 	TOptional<FInstallBundleCacheStats> Stats;
 
@@ -3906,7 +3907,7 @@ TOptional<FInstallBundleCacheStats> FDefaultInstallBundleManager::GetCacheStats(
 		if (pBundleCache)
 		{
 			const TSharedRef<FInstallBundleCache>& BundleCache = *pBundleCache;
-			Stats.Emplace(BundleCache->GetStats(DumpToLog, LogVerbosityOverride < ELogVerbosity::Verbose));
+			Stats.Emplace(BundleCache->GetStats(Flags, LogVerbosityOverride < ELogVerbosity::Verbose));
 		}
 	}
 
@@ -5083,7 +5084,7 @@ void FDefaultInstallBundleManager::AsyncInit_FireInitAnlaytic(bool bCanRetry)
 		{
 			InstallBundleManagerAnalytics::FireEvent_BundleManagerCacheStats(
 				AnalyticsProvider.Get(),
-				Pair.Value->GetStats(EInstallBundleCacheDumpToLog::Default));
+				Pair.Value->GetStats(EInstallBundleCacheStatsFlags::DumpToLog));
 		}
 	}
 
diff --git a/Engine/Plugins/Experimental/DefaultInstallBundleManager/Source/Public/DefaultInstallBundleManager.h b/Engine/Plugins/Experimental/DefaultInstallBundleManager/Source/Public/DefaultInstallBundleManager.h
--- a/Engine/Plugins/Experimental/DefaultInstallBundleManager/Source/Public/DefaultInstallBundleManager.h
+++ b/Engine/Plugins/Experimental/DefaultInstallBundleManager/Source/Public/DefaultInstallBundleManager.h
@@ -657,8 +657,8 @@ public:
 	
 	virtual EInstallBundleResult FlushCache(FInstallBundleSourceOrCache SourceOrCache, FInstallBundleManagerFlushCacheCompleteDelegate Callback, ELogVerbosity::Type LogVerbosityOverride = ELogVerbosity::NoLogging) override;
 
-	virtual TArray<FInstallBundleCacheStats> GetCacheStats(EInstallBundleCacheDumpToLog DumpToLog = EInstallBundleCacheDumpToLog::None, ELogVerbosity::Type LogVerbosityOverride = ELogVerbosity::NoLogging) override;
-	virtual TOptional<FInstallBundleCacheStats> GetCacheStats(FInstallBundleSourceOrCache SourceOrCache, EInstallBundleCacheDumpToLog DumpToLog = EInstallBundleCacheDumpToLog::None, ELogVerbosity::Type LogVerbosityOverride = ELogVerbosity::NoLogging) override;
+	virtual TArray<FInstallBundleCacheStats> GetCacheStats(EInstallBundleCacheStatsFlags Flags = EInstallBundleCacheStatsFlags::None, ELogVerbosity::Type LogVerbosityOverride = ELogVerbosity::NoLogging) override;
+	virtual TOptional<FInstallBundleCacheStats> GetCacheStats(FInstallBundleSourceOrCache SourceOrCache, EInstallBundleCacheStatsFlags Flags = EInstallBundleCacheStatsFlags::None, ELogVerbosity::Type LogVerbosityOverride = ELogVerbosity::NoLogging) override;
 
 	virtual void RequestRemoveContentOnNextInit(TArrayView<const FName> RemoveNames, TArrayView<const FName> KeepNames = TArrayView<const FName>()) override;
 
diff --git a/Engine/Source/Runtime/InstallBundleManager/Private/InstallBundleCache.cpp b/Engine/Source/Runtime/InstallBundleManager/Private/InstallBundleCache.cpp
--- a/Engine/Source/Runtime/InstallBundleManager/Private/InstallBundleCache.cpp
+++ b/Engine/Source/Runtime/InstallBundleManager/Private/InstallBundleCache.cpp
@@ -392,13 +392,14 @@ void FInstallBundleCache::CheckInvariants() const
 #endif // INSTALLBUNDLE_CACHE_CHECK_INVARIANTS
 }
 
-FInstallBundleCacheStats FInstallBundleCache::GetStats(EInstallBundleCacheDumpToLog DumpToLog /*= EInstallBundleCacheDumpToLog::None*/, bool bVerbose /*= false*/) const
+FInstallBundleCacheStats FInstallBundleCache::GetStats(EInstallBundleCacheStatsFlags Flags /*= EInstallBundleCacheStatsFlags::None*/, bool bVerbose /*= false*/) const
 {
 	FInstallBundleCacheStats Stats;
 	Stats.CacheName = CacheName;
 	Stats.MaxSize = TotalSize;
 
-	bool bDumpToLog = DumpToLog != EInstallBundleCacheDumpToLog::None;
+	const bool bDumpToLog = EnumHasAnyFlags(Flags, EInstallBundleCacheStatsFlags::DumpToLog);
+	const bool bDumpToResults = EnumHasAnyFlags(Flags, EInstallBundleCacheStatsFlags::DumpToResults);
 
 	if (bDumpToLog)
 	{
@@ -409,7 +410,7 @@ FInstallBundleCacheStats FInstallBundleCache::GetStats(EInstallBundleCacheDumpTo
 		CacheInfo.ValueSort(FCacheSortPredicate());
 	}
 
-	auto AddStats = [&Stats](const FBundleCacheInfo& Info)
+	auto AddStats = [&Stats, bDumpToResults](FName BundleName, const FBundleCacheInfo& Info)
 	{
 		Stats.UsedSize += Info.GetSize();
 
@@ -417,6 +418,20 @@ FInstallBundleCacheStats FInstallBundleCache::GetStats(EInstallBundleCacheDumpTo
 		{
 			Stats.ReservedSize += Info.CurrentInstallSize;
 		}
+
+		if (bDumpToResults)
+		{
+			Stats.BundleStats.Emplace(FInstallBundleCacheBundleStats
+			{
+				.BundleName = BundleName,
+				.FullInstallSize = Info.FullInstallSize,
+				.InstallOverheadSize = Info.InstallOverheadSize,
+				.CurrentInstallSize = Info.CurrentInstallSize,
+				.TimeStamp = Info.TimeStamp,
+				.AgeScalar = Info.AgeScalar,
+				.bReserved = Info.State == ECacheState::Reserved,
+			});
+		}
 	};
 
 	// Note: Yes this verbosity is correct and very non-intuitive
@@ -465,62 +480,59 @@ FInstallBundleCacheStats FInstallBundleCache::GetStats(EInstallBundleCacheDumpTo
 		INSTALLBUNDLECACHE_CSV_LOG(Display);
 	};
 
-	switch(DumpToLog)
+	if (!bDumpToLog)
 	{
-	default:
-	case EInstallBundleCacheDumpToLog::None:
 		for (const TPair<FName, FBundleCacheInfo>& CachePair : CacheInfo)
 		{
 			const FBundleCacheInfo& Info = CachePair.Value;
-			AddStats(Info);
+			AddStats(CachePair.Key, Info);
 		}
-		break;
-		
-	case EInstallBundleCacheDumpToLog::Default:
+	}
+	else if (EnumHasAnyFlags(Flags, EInstallBundleCacheStatsFlags::CSVFormat))
+	{
 		if (bVerbose)
 		{
+			INSTALLBUNDLECACHE_CSV_HEADER_LOG(Display);
+
 			for (const TPair<FName, FBundleCacheInfo>& CachePair : CacheInfo)
 			{
 				const FBundleCacheInfo& Info = CachePair.Value;
-				AddStats(Info);
-				DumpToLog_DefaultVerbose(CachePair.Key, Info);
+				AddStats(CachePair.Key, Info);
+				DumpToLog_CSVVerbose(CachePair.Key, Info);
 			}
 		}
 		else
 		{
+			INSTALLBUNDLECACHE_CSV_HEADER_LOG(Verbose);
+
 			for (const TPair<FName, FBundleCacheInfo>& CachePair : CacheInfo)
 			{
 				const FBundleCacheInfo& Info = CachePair.Value;
-				AddStats(Info);
-				DumpToLog_Default(CachePair.Key, Info);
+				AddStats(CachePair.Key, Info);
+				DumpToLog_CSV(CachePair.Key, Info);
 			}
 		}
-		break;
-
-	case EInstallBundleCacheDumpToLog::CSV:
+	}
+	else
+	{
 		if (bVerbose)
 		{
-			INSTALLBUNDLECACHE_CSV_HEADER_LOG(Display);
-
 			for (const TPair<FName, FBundleCacheInfo>& CachePair : CacheInfo)
 			{
 				const FBundleCacheInfo& Info = CachePair.Value;
-				AddStats(Info);
-				DumpToLog_CSVVerbose(CachePair.Key, Info);
+				AddStats(CachePair.Key, Info);
+				DumpToLog_DefaultVerbose(CachePair.Key, Info);
 			}
 		}
 		else
 		{
-			INSTALLBUNDLECACHE_CSV_HEADER_LOG(Verbose);
-
 			for (const TPair<FName, FBundleCacheInfo>& CachePair : CacheInfo)
 			{
 				const FBundleCacheInfo& Info = CachePair.Value;
-				AddStats(Info);
-				DumpToLog_CSV(CachePair.Key, Info);
+				AddStats(CachePair.Key, Info);
+				DumpToLog_Default(CachePair.Key, Info);
 			}
 		}
-		break;
 	}
 
 #undef INSTALLBUNDLECACHE_CSV_LOG
diff --git a/Engine/Source/Runtime/InstallBundleManager/Private/InstallBundleManagerInterface.cpp b/Engine/Source/Runtime/InstallBundleManager/Private/InstallBundleManagerInterface.cpp
--- a/Engine/Source/Runtime/InstallBundleManager/Private/InstallBundleManagerInterface.cpp
+++ b/Engine/Source/Runtime/InstallBundleManager/Private/InstallBundleManagerInterface.cpp
@@ -81,6 +81,38 @@ TValueOrError<FInstallBundleReleaseRequestInfo, EInstallBundleResult> IInstallBu
 	return RequestReleaseContent(MakeArrayView(&ReleaseName, 1), Flags, KeepNames, LogVerbosityOverride);
 }
 
+PRAGMA_DISABLE_DEPRECATION_WARNINGS
+TArray<FInstallBundleCacheStats> IInstallBundleManager::GetCacheStats(EInstallBundleCacheDumpToLog DumpToLog, ELogVerbosity::Type LogVerbosityOverride /*= ELogVerbosity::NoLogging*/)
+{
+	EInstallBundleCacheStatsFlags Flags = EInstallBundleCacheStatsFlags::None;
+	switch (DumpToLog)
+	{	
+	case EInstallBundleCacheDumpToLog::CSV:
+		Flags |= EInstallBundleCacheStatsFlags::CSVFormat;
+		// Fallthrough
+	case EInstallBundleCacheDumpToLog::Default:
+		Flags |= EInstallBundleCacheStatsFlags::DumpToLog;
+	}
+
+	return GetCacheStats(Flags, LogVerbosityOverride);
+}
+
+TOptional<FInstallBundleCacheStats> IInstallBundleManager::GetCacheStats(FInstallBundleSourceOrCache SourceOrCache, EInstallBundleCacheDumpToLog DumpToLog, ELogVerbosity::Type LogVerbosityOverride /*= ELogVerbosity::NoLogging*/)
+{
+	EInstallBundleCacheStatsFlags Flags = EInstallBundleCacheStatsFlags::None;
+	switch (DumpToLog)
+	{
+	case EInstallBundleCacheDumpToLog::CSV:
+		Flags |= EInstallBundleCacheStatsFlags::CSVFormat;
+		// Fallthrough
+	case EInstallBundleCacheDumpToLog::Default:
+		Flags |= EInstallBundleCacheStatsFlags::DumpToLog;
+	}
+
+	return GetCacheStats(SourceOrCache, Flags, LogVerbosityOverride);
+}
+PRAGMA_ENABLE_DEPRECATION_WARNINGS
+
 EInstallBundleResult IInstallBundleManager::FlushCache(FInstallBundleManagerFlushCacheCompleteDelegate Callback, ELogVerbosity::Type LogVerbosityOverride /*= ELogVerbosity::NoLogging*/)
 {
 	return FlushCache({}, MoveTemp(Callback), LogVerbosityOverride);
diff --git a/Engine/Source/Runtime/InstallBundleManager/Public/InstallBundleCache.h b/Engine/Source/Runtime/InstallBundleManager/Public/InstallBundleCache.h
--- a/Engine/Source/Runtime/InstallBundleManager/Public/InstallBundleCache.h
+++ b/Engine/Source/Runtime/InstallBundleManager/Public/InstallBundleCache.h
@@ -96,7 +96,7 @@ public:
 	// Hint to the cache that this bundle is requested, and we should prefer to evict non-requested bundles if possible
 	INSTALLBUNDLEMANAGER_API void HintRequested(FName BundleName, bool bRequested);
 
-	INSTALLBUNDLEMANAGER_API FInstallBundleCacheStats GetStats(EInstallBundleCacheDumpToLog DumpToLog = EInstallBundleCacheDumpToLog::None, bool bVerbose = false) const;
+	INSTALLBUNDLEMANAGER_API FInstallBundleCacheStats GetStats(EInstallBundleCacheStatsFlags Flags = EInstallBundleCacheStatsFlags::None, bool bVerbose = false) const;
 
 private:
 	INSTALLBUNDLEMANAGER_API uint64 GetFreeSpaceInternal(uint64 UsedSize) const;
diff --git a/Engine/Source/Runtime/InstallBundleManager/Public/InstallBundleManagerInterface.h b/Engine/Source/Runtime/InstallBundleManager/Public/InstallBundleManagerInterface.h
--- a/Engine/Source/Runtime/InstallBundleManager/Public/InstallBundleManagerInterface.h
+++ b/Engine/Source/Runtime/InstallBundleManager/Public/InstallBundleManagerInterface.h
@@ -167,8 +167,15 @@ public:
 	INSTALLBUNDLEMANAGER_API EInstallBundleResult FlushCache(FInstallBundleManagerFlushCacheCompleteDelegate Callback, ELogVerbosity::Type LogVerbosityOverride = ELogVerbosity::NoLogging);
 	virtual EInstallBundleResult FlushCache(FInstallBundleSourceOrCache SourceOrCache, FInstallBundleManagerFlushCacheCompleteDelegate Callback, ELogVerbosity::Type LogVerbosityOverride = ELogVerbosity::NoLogging) = 0;
 
-	virtual TArray<FInstallBundleCacheStats> GetCacheStats(EInstallBundleCacheDumpToLog DumpToLog = EInstallBundleCacheDumpToLog::None, ELogVerbosity::Type LogVerbosityOverride = ELogVerbosity::NoLogging) = 0;
-	virtual TOptional<FInstallBundleCacheStats> GetCacheStats(FInstallBundleSourceOrCache SourceOrCache, EInstallBundleCacheDumpToLog DumpToLog = EInstallBundleCacheDumpToLog::None, ELogVerbosity::Type LogVerbosityOverride = ELogVerbosity::NoLogging) = 0;
+PRAGMA_DISABLE_DEPRECATION_WARNINGS
+	UE_DEPRECATED(5.7, "Use overload with EInstallBundleCacheStatsFlags")
+	INSTALLBUNDLEMANAGER_API virtual TArray<FInstallBundleCacheStats> GetCacheStats(EInstallBundleCacheDumpToLog DumpToLog, ELogVerbosity::Type LogVerbosityOverride = ELogVerbosity::NoLogging);
+	UE_DEPRECATED(5.7, "Use overload with EInstallBundleCacheStatsFlags")
+	INSTALLBUNDLEMANAGER_API virtual TOptional<FInstallBundleCacheStats> GetCacheStats(FInstallBundleSourceOrCache SourceOrCache, EInstallBundleCacheDumpToLog DumpToLog, ELogVerbosity::Type LogVerbosityOverride = ELogVerbosity::NoLogging);
+PRAGMA_ENABLE_DEPRECATION_WARNINGS
+
+	virtual TArray<FInstallBundleCacheStats> GetCacheStats(EInstallBundleCacheStatsFlags Flags = EInstallBundleCacheStatsFlags::None, ELogVerbosity::Type LogVerbosityOverride = ELogVerbosity::NoLogging) = 0;
+	virtual TOptional<FInstallBundleCacheStats> GetCacheStats(FInstallBundleSourceOrCache SourceOrCache, EInstallBundleCacheStatsFlags Flags = EInstallBundleCacheStatsFlags::None, ELogVerbosity::Type LogVerbosityOverride = ELogVerbosity::NoLogging) = 0;
 
 	INSTALLBUNDLEMANAGER_API void RequestRemoveContentOnNextInit(FName RemoveName, TArrayView<const FName> KeepNames = TArrayView<const FName>());
 	virtual void RequestRemoveContentOnNextInit(TArrayView<const FName> RemoveNames, TArrayView<const FName> KeepNames = TArrayView<const FName>()) = 0;
diff --git a/Engine/Source/Runtime/InstallBundleManager/Public/InstallBundleTypes.h b/Engine/Source/Runtime/InstallBundleManager/Public/InstallBundleTypes.h
--- a/Engine/Source/Runtime/InstallBundleManager/Public/InstallBundleTypes.h
+++ b/Engine/Source/Runtime/InstallBundleManager/Public/InstallBundleTypes.h
@@ -451,6 +451,17 @@ enum class EInstallBundleSourceBundleSkipReason : uint32
 };
 ENUM_CLASS_FLAGS(EInstallBundleSourceBundleSkipReason);
 
+struct FInstallBundleCacheBundleStats
+{
+	FName BundleName;
+	uint64 FullInstallSize = 0;
+	uint64 InstallOverheadSize = 0;
+	uint64 CurrentInstallSize = 0;
+	FDateTime TimeStamp = FDateTime::MinValue();
+	double AgeScalar = 1.0;
+	bool bReserved = false;
+};
+
 struct FInstallBundleCacheStats
 {
 	FName CacheName;
@@ -458,9 +469,20 @@ struct FInstallBundleCacheStats
 	uint64 UsedSize = 0;
 	uint64 ReservedSize = 0;
 	uint64 FreeSize = 0;
+	// EInstallBundleCacheStatsFlags::DumpToResults must be used to populate BundleStats
+	TArray<FInstallBundleCacheBundleStats> BundleStats;
+};
+
+enum class EInstallBundleCacheStatsFlags : uint8
+{
+	None = 0,
+	DumpToLog = (1 << 0),
+	CSVFormat = (1 << 1),
+	DumpToResults = (1 << 2)
 };
+ENUM_CLASS_FLAGS(EInstallBundleCacheStatsFlags);
 
-enum class EInstallBundleCacheDumpToLog : int8
+enum class UE_DEPRECATED(5.7, "Use EInstallBundleCacheStatsFlags") EInstallBundleCacheDumpToLog : int8
 {
 	None = 0,
 	Default,
diff --git a/Engine/Source/Runtime/NullInstallBundleManager/Private/NullInstallBundleManager.cpp b/Engine/Source/Runtime/NullInstallBundleManager/Private/NullInstallBundleManager.cpp
--- a/Engine/Source/Runtime/NullInstallBundleManager/Private/NullInstallBundleManager.cpp
+++ b/Engine/Source/Runtime/NullInstallBundleManager/Private/NullInstallBundleManager.cpp
@@ -84,12 +84,24 @@ class FNullInstallBundleManager final : public IInstallBundleManager
 		return EInstallBundleResult::OK;
 	}
 
-	virtual TArray<FInstallBundleCacheStats> GetCacheStats(EInstallBundleCacheDumpToLog DumpToLog = EInstallBundleCacheDumpToLog::None, ELogVerbosity::Type LogVerbosityOverride = ELogVerbosity::NoLogging) override
+PRAGMA_DISABLE_DEPRECATION_WARNINGS
+	virtual TArray<FInstallBundleCacheStats> GetCacheStats(EInstallBundleCacheDumpToLog DumpToLog, ELogVerbosity::Type LogVerbosityOverride = ELogVerbosity::NoLogging) override
 	{
 		return TArray<FInstallBundleCacheStats>();
 	}
 
-	virtual TOptional<FInstallBundleCacheStats> GetCacheStats(FInstallBundleSourceOrCache SourceOrCache, EInstallBundleCacheDumpToLog DumpToLog = EInstallBundleCacheDumpToLog::None, ELogVerbosity::Type LogVerbosityOverride = ELogVerbosity::NoLogging) override
+	virtual TOptional<FInstallBundleCacheStats> GetCacheStats(FInstallBundleSourceOrCache SourceOrCache, EInstallBundleCacheDumpToLog DumpToLog, ELogVerbosity::Type LogVerbosityOverride = ELogVerbosity::NoLogging) override
+	{
+		return TOptional<FInstallBundleCacheStats>();
+	}
+PRAGMA_ENABLE_DEPRECATION_WARNINGS
+
+	virtual TArray<FInstallBundleCacheStats> GetCacheStats(EInstallBundleCacheStatsFlags Flags = EInstallBundleCacheStatsFlags::None, ELogVerbosity::Type LogVerbosityOverride = ELogVerbosity::NoLogging) override
+	{
+		return TArray<FInstallBundleCacheStats>();
+	}
+
+	virtual TOptional<FInstallBundleCacheStats> GetCacheStats(FInstallBundleSourceOrCache SourceOrCache, EInstallBundleCacheStatsFlags Flags = EInstallBundleCacheStatsFlags::None, ELogVerbosity::Type LogVerbosityOverride = ELogVerbosity::NoLogging) override
 	{
 		return TOptional<FInstallBundleCacheStats>();
 	}
