Refactor InstallBundleSourceBPS::GetContentState
- Use a single async task per tick instead of one per request
- Dedupe install bundles across all processed requests, previously requests with the same bundles would lead to redundant processing.
- Deprecate ForceNoPatching flag.
- Remove HACK_ALLOW_SOURCE_BPS_SOFTFAIL

diff --git a/Engine/Source/Runtime/InstallBundleManager/Public/InstallBundleTypes.h b/Engine/Source/Runtime/InstallBundleManager/Public/InstallBundleTypes.h
--- a/Engine/Source/Runtime/InstallBundleManager/Public/InstallBundleTypes.h
+++ b/Engine/Source/Runtime/InstallBundleManager/Public/InstallBundleTypes.h
@@ -191,7 +191,7 @@ struct FInstallBundleCombinedContentState
 enum class EInstallBundleGetContentStateFlags : uint32
 {
 	None = 0,
-	ForceNoPatching = (1 << 0),
+	ForceNoPatching UE_DEPRECATED(5.7, "ForceNoPatching is no longer supported, use uninstall") = (1 << 0),
 };
 ENUM_CLASS_FLAGS(EInstallBundleGetContentStateFlags);
 
@@ -237,7 +237,7 @@ enum class EInstallBundleRequestFlags : uint32
 	CheckForCellularDataUsage = (1 << 0),
 	UseBackgroundDownloads = (1 << 1),
 	SendNotificationIfDownloadCompletesInBackground = (1 << 2),
-	ForceNoPatching = (1 << 3),
+	ForceNoPatching UE_DEPRECATED(5.7, "ForceNoPatching is no longer supported, use uninstall") = (1 << 3),
 	TrackPersistentBundleStats = (1 << 4),
 	SkipMount = (1 << 5),
 	AsyncMount = (1 << 6),
