Add SpaceRequiredForInstallOtherDirs to FInstallBundleContentSize: on some platforms, install bundles may be installed in partitions that are different from GamePersistentDownloadDir, and FInstallBundleCombinedContentState::FreeSpace only accounts for the free space in GamePersistentDownloadDir

This way, the existing logic that compares FInstallBundleContentSize::SpaceRequiredForInstall to FInstallBundleCombinedContentState::FreeSpace is still correct, and if one wants to know the system wide size for an install bundle, they need to do SpaceRequiredForInstall + SpaceRequiredForInstallOtherDirs

diff --git a/Engine/Plugins/Experimental/DefaultInstallBundleManager/Source/Private/InstallBundleManagerUtil.cpp b/Engine/Plugins/Experimental/DefaultInstallBundleManager/Source/Private/InstallBundleManagerUtil.cpp
--- a/Engine/Plugins/Experimental/DefaultInstallBundleManager/Source/Private/InstallBundleManagerUtil.cpp
+++ b/Engine/Plugins/Experimental/DefaultInstallBundleManager/Source/Private/InstallBundleManagerUtil.cpp
@@ -317,8 +317,8 @@ namespace InstallBundleManagerUtil
 	{
 		FSessionAnalyticsData& FoundData = SessionAnalyticsDataMap.FindOrAdd(SessionName);
 		FoundData.bRequiresDownload = ContentState.ContentSize.DownloadSize > 0;
-		FoundData.bRequiresInstall = ContentState.ContentSize.SpaceRequiredForInstall > 0;
-		
+		FoundData.bRequiresInstall = ContentState.ContentSize.SpaceRequiredForInstall + ContentState.ContentSize.SpaceRequiredForInstallOtherDirs> 0;
+
 		//Reset so that it can react to the above values being changed
 		FoundData.ResetShouldSendBGAnalytics();
 	}
diff --git a/Engine/Source/Runtime/InstallBundleManager/Public/InstallBundleTypes.h b/Engine/Source/Runtime/InstallBundleManager/Public/InstallBundleTypes.h
--- a/Engine/Source/Runtime/InstallBundleManager/Public/InstallBundleTypes.h
+++ b/Engine/Source/Runtime/InstallBundleManager/Public/InstallBundleTypes.h
@@ -131,8 +131,10 @@ struct FInstallBundleContentSize
 {
 	// Size remaining to download
 	uint64 DownloadSize = 0;
-	// Size needed to install the remaining download.
+	// Size needed to install the remaining download under GamePersistentDownloadDir
 	uint64 SpaceRequiredForInstall = 0;
+	// Size needed to install the remaining download under other directories GamePersistentDownloadDir
+	uint64 SpaceRequiredForInstallOtherDirs = 0;
 	// size of bundle currently on disk.
 	uint64 CurrentSizeOnDisk = 0;
 
@@ -141,6 +143,7 @@ struct FInstallBundleContentSize
 		FInstallBundleContentSize Result;
 		Result.DownloadSize = DownloadSize + Other.DownloadSize;
 		Result.SpaceRequiredForInstall = SpaceRequiredForInstall + Other.SpaceRequiredForInstall;
+		Result.SpaceRequiredForInstallOtherDirs = SpaceRequiredForInstallOtherDirs + Other.SpaceRequiredForInstallOtherDirs;
 		Result.CurrentSizeOnDisk = CurrentSizeOnDisk + Other.CurrentSizeOnDisk;
 		return Result;
 	}
