New InstallBundleManager reporting

#rnx

diff --git a/Engine/Plugins/Experimental/DefaultInstallBundleManager/Source/Private/DefaultInstallBundleManager.cpp b/Engine/Plugins/Experimental/DefaultInstallBundleManager/Source/Private/DefaultInstallBundleManager.cpp
--- a/Engine/Plugins/Experimental/DefaultInstallBundleManager/Source/Private/DefaultInstallBundleManager.cpp
+++ b/Engine/Plugins/Experimental/DefaultInstallBundleManager/Source/Private/DefaultInstallBundleManager.cpp
@@ -189,6 +189,7 @@ FDefaultInstallBundleManager::FContentRequest::~FContentRequest() = default;
 FDefaultInstallBundleManager::FDefaultInstallBundleManager(FInstallBundleSourceFactoryFunction InInstallBundleSourceFactory)
 	: InstallBundleSourceFactory(InInstallBundleSourceFactory ? MoveTemp(InInstallBundleSourceFactory) : InstallBundleManagerUtil::MakeBundleSource )
 	, PersistentStats(MakeShared<InstallBundleManagerUtil::FPersistentStatContainer>())
+	, PersistentStateCache(AsyncIOTasks)
 	, AnalyticsProvider(nullptr)
 	, StatsMap(MakeShared<InstallBundleUtil::FContentRequestStatsMap>())
 {
@@ -209,12 +210,16 @@ FDefaultInstallBundleManager::~FDefaultInstallBundleManager()
 	FTSTicker::GetCoreTicker().RemoveTicker(TickHandle);
 	TickHandle.Reset();
 
-	InstallBundleUtil::CleanupInstallBundleAsyncIOTasks(AsyncMountTasks);
+	FTSTicker::GetCoreTicker().RemoveTicker(ReportingHandle);
+	ReportingHandle.Reset();
+
+	InstallBundleUtil::CleanupInstallBundleAsyncIOTasks(AsyncIOTasks);
 }
 
 void FDefaultInstallBundleManager::Initialize()
 {
 	TickHandle = FTSTicker::GetCoreTicker().AddTicker(FTickerDelegate::CreateRaw(this, &FDefaultInstallBundleManager::Tick));
+	ReportingHandle = FTSTicker::GetCoreTicker().AddTicker(FTickerDelegate::CreateRaw(this, &FDefaultInstallBundleManager::TickReporting), 2.0f);
 
 	InitErrorHandlerStack.Push(FInstallBundleManagerInitErrorHandler::CreateLambda(
 		[](EInstallBundleManagerInitResult Error)
@@ -262,7 +267,7 @@ bool FDefaultInstallBundleManager::Tick(float dt)
 
 	TickPauseStatus(false);
 
-	TickAsyncMountTasks();
+	TickAsyncIOTasks();
 
 	TickReleaseRequests();
 
@@ -421,6 +426,10 @@ void FDefaultInstallBundleManager::TickInit()
 			InitResult = EInstallBundleManagerInitResult::OK;
 			AsyncInit_InitBundleCaches();
 			break;
+		case EAsyncInitStep::InitPersisentReportCache:
+			InitResult = EInstallBundleManagerInitResult::OK;
+			AsyncInit_InitPersistentReportCache();
+			break;
 		case EAsyncInitStep::QueryBundleInfo:
 			InitResult = EInstallBundleManagerInitResult::OK;
 			AsyncInit_QueryBundleInfo();
@@ -1556,11 +1565,11 @@ void FDefaultInstallBundleManager::UpdateBundleSourceReleaseComplete(TSharedRef<
 	Request->StepResult = EContentRequestStepResult::Done;
 }
 
-void FDefaultInstallBundleManager::TickAsyncMountTasks()
+void FDefaultInstallBundleManager::TickAsyncIOTasks()
 {
 	CSV_SCOPED_TIMING_STAT(InstallBundleManager, InstallBundleManager_TickAsyncMountTasks);
 
-	InstallBundleUtil::FinishInstallBundleAsyncIOTasks(AsyncMountTasks);
+	InstallBundleUtil::FinishInstallBundleAsyncIOTasks(AsyncIOTasks);
 }
 
 void FDefaultInstallBundleManager::MountPaks(FContentRequestRef Request)
@@ -2539,6 +2548,21 @@ void FDefaultInstallBundleManager::TickPruneBundleInfo()
 	}
 }
 
+bool FDefaultInstallBundleManager::TickReporting(float dt)
+{
+	// If it it hasn't changed, don't do anything
+	if (bPersistentStateDirty)
+	{
+		bPersistentStateDirty = false;
+
+		// Update the cache now (so it's saved to disk), and broadcast to listeners
+		PersistentStateCache.SetReport(PersistentState);
+		ReportingDelegate.Broadcast(PersistentState);
+		//FPlatformMisc::LowLevelOutputDebugStringf(TEXT("InstallBundleManager: %s"), *PersistentState.GetLongProgressString());
+	}
+	return true;
+}
+
 void FDefaultInstallBundleManager::IterateContentRequests(TFunctionRef<bool(const FContentRequestRef& QueuedRequest)> OnFound)
 {
 	bool bContinue = true;
@@ -4565,6 +4589,51 @@ EInstallBundleManagerInitResult FDefaultInstallBundleManager::Init_TryCreateBund
 			{
 				check(SourceType == Source->GetSourceType());
 				Source->SetErrorSimulationCommands(FCommandLine::Get());
+				if (Source->SupportsReportingDelegate())
+				{
+					Source->SetReportingDelegate(FInstallBundleReportingDelegate::CreateLambda([this, WeakSource = Source.ToWeakPtr()](FInstallBundleReport BundleReport)
+						{
+							checkf(IsInGameThread(), TEXT("ReportingDelegate may only be called on game thread"));
+							TSharedPtr<IInstallBundleSource> PinnedSource = WeakSource.Pin();
+							if (PinnedSource.IsValid())
+							{
+								FInstallBundleSourceReport& PersistentSourceReport = PersistentState.FindOrAddBundleSourceReport(PinnedSource->GetSourceType());
+								FInstallBundleReport& PersistentReport = PersistentSourceReport.FindOrAddBundleReport(BundleReport.BundleName);
+
+								const int32 PersistentIndex = PersistentReport.State.IsSet() ? PersistentReport.State->GetIndex() : -1;
+								const int32 NewIndex = BundleReport.State.IsSet() ? BundleReport.State->GetIndex() : -1;
+
+								if (PersistentIndex != NewIndex)
+								{
+									LOG_INSTALL_BUNDLE_MAN(Verbose, TEXT("InstallBundleManager: bundle %s changed type index from %d to %d"), *BundleReport.BundleName.ToString(), PersistentIndex, NewIndex);
+								}
+
+								// We blindly copy what comes in if there is no state set yet in the persistent data
+								if (!PersistentReport.State.IsSet())
+								{
+									PersistentReport = BundleReport;
+									bPersistentStateDirty = true;
+								}
+								// Otherwise, if the persistent state is not set to FStateUpdated, we take whatever came in, but we add the difference
+								// in total to compensate for a source potentially reporting less than what was truly needed
+								else if (!PersistentReport.State->IsType<FInstallBundleReport::FStateUpdated>())
+								{
+									const uint64 TotalPersistent = PersistentReport.TotalDownloadSize();
+									const uint64 TotalNew = BundleReport.TotalDownloadSize();
+									if (ensureMsgf(TotalNew <= TotalPersistent, TEXT("Bundle %s has TotalPersistent=%" UINT64_FMT ", and TotalNew=%" UINT64_FMT ", which is not expected"), *BundleReport.BundleName.ToString(), TotalPersistent, TotalNew))
+									{
+										if (const uint64 Discrepancy = TotalPersistent - TotalNew)
+										{
+											BundleReport.AddDiscrepancy(Discrepancy);
+										}
+									}
+									PersistentReport = BundleReport;
+									bPersistentStateDirty = true;
+								}
+
+							}
+						}));
+				}
 				BundleSources.Add(SourceType, Source);
 				if (OutNewSources)
 				{
@@ -4911,6 +4980,17 @@ void FDefaultInstallBundleManager::AsyncInit_InitBundleCaches()
 	InitStepResult = EAsyncInitStepResult::Done;
 }
 
+void FDefaultInstallBundleManager::AsyncInit_InitPersistentReportCache()
+{
+	LOG_INSTALL_BUNDLE_MAN(Display, TEXT("Initializing persistent report cache..."));
+	InitStepResult = EAsyncInitStepResult::Waiting;
+	PersistentStateCache.Load(FInstallManagerBundleReportCacheDoneLoadingDelegate::CreateLambda([this]() {
+		LOG_INSTALL_BUNDLE_MAN(Display, TEXT("Persistent report cache context set to [%s]"), *PersistentState.ContextName.ToString());
+		PersistentState = PersistentStateCache.GetReport();
+		InitStepResult = EAsyncInitStepResult::Done;
+	}));
+}
+
 void FDefaultInstallBundleManager::AsyncInit_QueryBundleInfo()
 {
 	LOG_INSTALL_BUNDLE_MAN(Display, TEXT("Querying bundle sources for bundle info..."));
@@ -5171,7 +5251,61 @@ void FDefaultInstallBundleManager::LogStats(const InstallBundleUtil::FContentReq
 	InstallBundleManagerUtil::LogBundleRequestStats(Key, RequestStats, LogVerbosityOverride);
 }
 
+void FDefaultInstallBundleManager::SetContext(FName NewContextName) {
+	LOG_INSTALL_BUNDLE_MAN(Display, TEXT("Changing context to [%s]"), *NewContextName.ToString());
 
+	// Check prerequisites
+	checkf(InitState == EInstallBundleManagerInitState::Succeeded, TEXT("Init state must be set to Succeeded when setting a new context"));
+	checkf(ContextName == FName() || NewContextName == FName(), TEXT("Cannot have active context when setting a new context"));
+	ContextName = NewContextName;
+
+	// If we have the same context as the last time, check if any of our sources are planning to update to a different version 
+	if (PersistentState.ContextName == ContextName)
+	{
+		bool bFoundTargetVersionMismatch = false;
+		for (const TPair<FInstallBundleSourceType, TSharedPtr<IInstallBundleSource>>& Pair : BundleSources)
+		{
+			if (Pair.Value->SupportsReportingDelegate())
+			{
+				FInstallBundleSourceReport* BundleSourceReport = PersistentState.TryFindBundleSourceReport(Pair.Key);
+				if (!BundleSourceReport || !BundleSourceReport->TargetVersion.IsSet() || BundleSourceReport->TargetVersion.GetValue() != Pair.Value->GetContentVersion())
+				{
+					bFoundTargetVersionMismatch = true;
+					break;
+				}
+			}
+		}
+
+		// The context name is the same as last time, and all bundles still updating to the same version, we can continue
+		// with this persistent state.
+		if (!bFoundTargetVersionMismatch)
+		{
+			LOG_INSTALL_BUNDLE_MAN(Display, TEXT("Keep persistent state for context [%s]"), *ContextName.ToString());
+			return;
+		}
+	}
+
+	// Rebuild the persistent state from scratch
+	LOG_INSTALL_BUNDLE_MAN(Display, TEXT("Resetting persistent state to context [%s]"), *ContextName.ToString());
+	PersistentState = {};
+	PersistentState.ContextName = ContextName;
+	PersistentState.SessionId = FGuid::NewGuid().ToString();
+	for (const TPair<FInstallBundleSourceType, TSharedPtr<IInstallBundleSource>>& Pair : BundleSources)
+	{
+		if (Pair.Value->SupportsReportingDelegate())
+		{
+			PersistentState.FindOrAddBundleSourceReport(Pair.Key).TargetVersion = Pair.Value->GetContentVersion();
+		}
+	}
+
+	// NOTE: we don't call FPersistentStateCache::SetReport, so that we don't store the XML file if the system isn't otherwise used
+}
+
+bool FDefaultInstallBundleManager::CopyReportTo(FInstallManagerBundleReport& Report)
+{
+	Report = PersistentState;
+	return true;
+}
 
 const TCHAR* LexToString(FDefaultInstallBundleManager::EContentRequestBatch Val)
 {
@@ -5196,9 +5330,6 @@ const TCHAR* LexToString(FDefaultInstallBundleManager::EContentReleaseRequestBat
 	return InstallBundleUtil::TLexToString(Val, Strings);
 }
 
-
-
-
 class FDefaultInstallBundleManagerModule : public TInstallBundleManagerModule<FDefaultInstallBundleManager>
 {
 public:
@@ -5209,4 +5340,3 @@ public:
 };
 
 IMPLEMENT_GAME_MODULE(FDefaultInstallBundleManagerModule, DefaultInstallBundleManager)
-
diff --git a/Engine/Plugins/Experimental/DefaultInstallBundleManager/Source/Public/DefaultInstallBundleManager.h b/Engine/Plugins/Experimental/DefaultInstallBundleManager/Source/Public/DefaultInstallBundleManager.h
--- a/Engine/Plugins/Experimental/DefaultInstallBundleManager/Source/Public/DefaultInstallBundleManager.h
+++ b/Engine/Plugins/Experimental/DefaultInstallBundleManager/Source/Public/DefaultInstallBundleManager.h
@@ -5,6 +5,7 @@
 #include "InstallBundleManagerInterface.h"
 #include "InstallBundleManagerUtil.h"
 #include "InstallBundleCache.h"
+#include "InstallBundleManagerReporting.h"
 
 #include "Experimental/UnifiedError/UnifiedError.h"
 
@@ -85,6 +86,7 @@ protected:
 		None,
 		InitBundleSources,
 		InitBundleCaches,
+		InitPersisentReportCache,
 		QueryBundleInfo,
 		SetUpdateBundleInfoCallback,
 		CreateAnalyticsSession,
@@ -99,6 +101,7 @@ protected:
 			TEXT("None"),
 			TEXT("InitBundleSources"),
 			TEXT("InitBundleCaches"),
+			TEXT("InitPersisentReportCache"),
 			TEXT("QueryBundleInfo"),
 			TEXT("SetUpdateBundleInfoCallback"),
 			TEXT("CreateAnalyticsSession"),
@@ -597,12 +600,14 @@ protected:
 
 	UE_API void TickPauseStatus(bool bForceCallback);
 
-	UE_API void TickAsyncMountTasks();
+	UE_API void TickAsyncIOTasks();
 
 	UE_API void TickReleaseRequests();
 
 	UE_API void TickPruneBundleInfo();
 
+	UE_API bool TickReporting(float dt);
+
 	UE_API void IterateContentRequests(TFunctionRef<bool(const FContentRequestRef& QueuedRequest)> OnFound);
 	UE_API void IterateReleaseRequests(TFunctionRef<bool(const FContentReleaseRequestRef& QueuedRequest)> OnFound);
 
@@ -637,6 +642,8 @@ public:
 
 	UE_API virtual const TSharedPtr<IInstallBundleSource> GetBundleSource(FInstallBundleSourceType SourceType) const override;
 
+	UE_API virtual void SetContext(FName ContextName) override;
+	UE_API virtual bool CopyReportTo(FInstallManagerBundleReport& Report) override;
 
 	UE_API virtual FDelegateHandle PushInitErrorCallback(FInstallBundleManagerInitErrorHandler Callback) override;
 	UE_API virtual void PopInitErrorCallback() override;
@@ -730,6 +737,7 @@ protected:
 	UE_API void AsyncInit_InitBundleSources();
 	UE_API void AsyncInit_OnBundleSourceInitComplete(TSharedRef<IInstallBundleSource> Source, FInstallBundleSourceAsyncInitInfo InInitInfo);
 	UE_API void AsyncInit_InitBundleCaches();
+	UE_API void AsyncInit_InitPersistentReportCache();
 	UE_API void AsyncInit_QueryBundleInfo();
 	UE_API void AsyncInit_OnQueryBundleInfoComplete(TSharedRef<IInstallBundleSource> Source, FInstallBundleSourceBundleInfoQueryResult Result);
 	UE_API void AsyncInit_OnQueryBundleInfoComplete_HandleClientPatchCheck(EPatchCheckResult Result);
@@ -745,6 +753,7 @@ protected:
 
 protected:
 	FTSTicker::FDelegateHandle TickHandle;
+	FTSTicker::FDelegateHandle ReportingHandle;
 	FDelegateHandle AsyncInit_PatchCheckHandle;
 	FDelegateHandle PatchCheckHandle;
 	UE::IoStore::IOnDemandIoStore* OnDemandIoStore = nullptr;
@@ -795,7 +804,7 @@ protected:
 
 	TSharedRef<InstallBundleManagerUtil::FPersistentStatContainer> PersistentStats;
 
-	TArray<TUniquePtr<InstallBundleUtil::FInstallBundleTask>> AsyncMountTasks;
+	TArray<TUniquePtr<InstallBundleUtil::FInstallBundleTask>> AsyncIOTasks;
 	
 	bool bIsCheckingForPatch = false;
 	bool bDelayCheckingForContentPatch = false;
@@ -812,6 +821,12 @@ protected:
 
 	bool bHasEverUpdatedContent = false;
 
+	// Persistent update state
+	bool bPersistentStateDirty = false;
+	FName ContextName;
+	FInstallManagerBundleReport PersistentState;
+	FInstallManagerBundleReportCache PersistentStateCache; // must be declared after AsyncIOTasks
+
 protected:
 	// Analytics
 	TSharedPtr<IAnalyticsProviderET> AnalyticsProvider;
diff --git a/Engine/Source/Runtime/InstallBundleManager/Private/InstallBundleManagerInterface.cpp b/Engine/Source/Runtime/InstallBundleManager/Private/InstallBundleManagerInterface.cpp
--- a/Engine/Source/Runtime/InstallBundleManager/Private/InstallBundleManagerInterface.cpp
+++ b/Engine/Source/Runtime/InstallBundleManager/Private/InstallBundleManagerInterface.cpp
@@ -2,6 +2,7 @@
 
 #include "InstallBundleManagerInterface.h"
 #include "InstallBundleManagerModule.h"
+#include "InstallBundleManagerReporting.h"
 
 FInstallBundleManagerInitCompleteMultiDelegate IInstallBundleManager::InitCompleteDelegate;
 
@@ -15,6 +16,8 @@ FInstallBundleReleasedMultiDelegate IInstallBundleManager::ReleasedDelegate;
 
 FInstallBundleManagerOnPatchCheckComplete IInstallBundleManager::PatchCheckCompleteDelegate;
 
+FInstallBundleManagerReportingDelegate IInstallBundleManager::ReportingDelegate;
+
 TSharedPtr<IInstallBundleManager> IInstallBundleManager::GetPlatformInstallBundleManager()
 {
 	static IInstallBundleManagerModule* Module = nullptr;
@@ -50,12 +53,20 @@ TSharedPtr<IInstallBundleManager> IInstallBundleManager::GetPlatformInstallBundl
 	return {};
 }
 
+void IInstallBundleManager::SetContext(FName ContextName) 
+{
+}
+
+bool IInstallBundleManager::CopyReportTo(FInstallManagerBundleReport& Report)
+{ 
+	return false; 
+}
+
 const TSharedPtr<IInstallBundleSource> IInstallBundleManager::GetBundleSource(FInstallBundleSourceType SourceType) const
 {
 	return {};
 }
 
-
 TValueOrError<FInstallBundleRequestInfo, EInstallBundleResult> IInstallBundleManager::RequestUpdateContent(FName BundleName, EInstallBundleRequestFlags Flags, ELogVerbosity::Type LogVerbosityOverride /*= ELogVerbosity::NoLogging*/,  InstallBundleUtil::FContentRequestSharedContextPtr RequestSharedContext /*= nullptr*/)
 {
 	return RequestUpdateContent(MakeArrayView(&BundleName, 1), Flags, LogVerbosityOverride, RequestSharedContext);
diff --git a/Engine/Source/Runtime/InstallBundleManager/Private/InstallBundleManagerReporting.cpp b/Engine/Source/Runtime/InstallBundleManager/Private/InstallBundleManagerReporting.cpp
new file mode 100644
--- /dev/null
+++ b/Engine/Source/Runtime/InstallBundleManager/Private/InstallBundleManagerReporting.cpp
@@ -0,0 +1,403 @@
+// Copyright Epic Games, Inc. All Rights Reserved.
+
+#include "InstallBundleManagerReporting.h"
+#include "InstallBundleManagerPrivate.h"
+#include "Misc/FileHelper.h"
+#include "Misc/Paths.h"
+#include "Algo/Find.h"
+
+uint64 FInstallBundleReport::TotalDownloadSize() const
+{
+	if (!State.IsSet())
+	{
+		return 0;
+	}
+	if (const FStateUpdatable* StateUpdatable = State->TryGet<FStateUpdatable>())
+	{
+		return StateUpdatable->DownloadSize;
+	}
+	if (const FStateUpdating* StateUpdating = State->TryGet<FStateUpdating>())
+	{
+		return StateUpdating->DownloadSize;
+	}
+	if (const FStateInstalling* StateInstalling = State->TryGet<FStateInstalling>())
+	{
+		return StateInstalling->DownloadedSize;
+	}
+	if (const FStateUpdated* StateUpdated = State->TryGet<FStateUpdated>())
+	{
+		return StateUpdated->DownloadedSize;
+	}
+	checkf(false, TEXT("Unsupported state type in TotalDownloadSize()"));
+	return 0;
+}
+
+uint64 FInstallBundleReport::DownloadedBytes() const
+{
+	if (!State.IsSet() || State->IsType<FStateUpdatable>())
+	{
+		return 0;
+	}
+	if (const FStateUpdating* StateUpdating = State->TryGet<FStateUpdating>())
+	{
+		return static_cast<uint64>(FMath::RoundToInt64(StateUpdating->DownloadProgress * StateUpdating->DownloadSize));
+	}
+	if (const FStateInstalling* StateInstalling = State->TryGet<FStateInstalling>())
+	{
+		return StateInstalling->DownloadedSize;
+	}
+	if (const FStateUpdated* StateUpdated = State->TryGet<FStateUpdated>())
+	{
+		return StateUpdated->DownloadedSize;
+	}
+	checkf(false, TEXT("Unsupported state type in DownloadedBytes()"));
+	return 0;
+}
+
+float FInstallBundleReport::InstallationProgress() const
+{
+	if (!State.IsSet() || State->IsType<FStateUpdatable>() || State->IsType<FStateUpdating>())
+	{
+		return 0.f;
+	}
+	if (const FStateInstalling* StateInstalling = State->TryGet<FStateInstalling>())
+	{
+		check(StateInstalling->InstallProgress >= 0 && StateInstalling->InstallProgress <= 1.f);
+		return StateInstalling->InstallProgress;
+	}
+	if (State->IsType<FStateUpdated>())
+	{
+		return 1.f;
+	}
+	checkf(false, TEXT("Unsupported state type in DownloadedBytes()"));
+	return 0;
+}
+
+void FInstallBundleReport::AddDiscrepancy(uint64 Bytes)
+{
+	if (!Bytes || !State.IsSet())
+	{
+		return;
+	}
+	if (FStateUpdatable* StateUpdatable = State->TryGet<FStateUpdatable>())
+	{
+		StateUpdatable->DownloadSize += Bytes;
+	}
+	else if (FStateUpdating* StateUpdating = State->TryGet<FStateUpdating>())
+	{
+		const uint64 DownloadedBytes = StateUpdating->DownloadProgress * StateUpdating->DownloadSize + Bytes;
+		StateUpdating->DownloadSize += Bytes;
+		StateUpdating->DownloadProgress = DownloadedBytes / float(StateUpdating->DownloadSize);
+	}
+	else if (FStateInstalling* StateInstalling = State->TryGet<FStateInstalling>())
+	{
+		StateInstalling->DownloadedSize += Bytes;
+	}
+	else if (FStateUpdated* StateInstalled = State->TryGet<FStateUpdated>())
+	{
+		StateInstalled->DownloadedSize += Bytes;
+	}
+	else
+	{
+		checkf(false, TEXT("Unsupported state type in DownloadedBytes()"));
+	}
+}
+
+FInstallBundleReport* FInstallBundleSourceReport::TryFindBundleReport(FName BundleName)
+{
+	return Algo::FindByPredicate(Bundles, [&](const FInstallBundleReport& BundleReport) {
+		return BundleReport.BundleName == BundleName;
+	});
+}
+
+FInstallBundleReport& FInstallBundleSourceReport::FindOrAddBundleReport(FName BundleName)
+{
+	if (FInstallBundleReport* BundleReport = TryFindBundleReport(BundleName))
+	{
+		return *BundleReport;
+	}
+
+	FInstallBundleReport& NewReport = Bundles.Emplace_GetRef();
+	NewReport.BundleName = BundleName;
+	return *TryFindBundleReport(BundleName);
+}
+
+uint64 FInstallBundleSourceReport::TotalDownloadSize() const
+{
+	uint64 Total = 0;
+	for (const FInstallBundleReport& BundleReport : Bundles)
+	{
+		Total += BundleReport.TotalDownloadSize();
+	}
+	return Total;
+}
+
+uint64 FInstallBundleSourceReport::DownloadedBytes() const
+{
+	uint64 Total = 0;
+	for (const FInstallBundleReport& BundleReport : Bundles)
+	{
+		Total += BundleReport.DownloadedBytes();
+	}
+	return Total;
+}
+
+int FInstallBundleSourceReport::TotalBundleCount() const
+{
+	return Bundles.Num();
+}
+
+int FInstallBundleSourceReport::UpdatedBundleCount() const
+{
+	int Num = 0;
+	for (const FInstallBundleReport& BundleReport : Bundles)
+	{
+		if (BundleReport.IsUpdated())
+		{
+			++Num;
+		}
+	}
+	return Num;
+}
+
+FInstallBundleSourceReport* FInstallManagerBundleReport::TryFindBundleSourceReport(const FInstallBundleSourceType& Type)
+{
+	return Algo::FindByPredicate(SourceReports, [SourceType = FName(Type.GetNameStr())](const FInstallBundleSourceReport& BundleSourceReport) {
+		return BundleSourceReport.SourceType == SourceType;
+	});
+}
+
+FInstallBundleSourceReport& FInstallManagerBundleReport::FindOrAddBundleSourceReport(const FInstallBundleSourceType& Type)
+{
+	if (FInstallBundleSourceReport* BundleSourceReport = TryFindBundleSourceReport(Type))
+	{
+		return *BundleSourceReport;
+	}
+
+	FInstallBundleSourceReport& NewReport = SourceReports.Emplace_GetRef();
+	NewReport.SourceType = FName(Type.GetNameStr());
+	return *TryFindBundleSourceReport(Type);
+}
+
+float FInstallManagerBundleReport::GetDownloadProgress() const
+{
+	const uint64 TotalBytes = TotalDownloadSize();
+	if (!TotalBytes)
+	{
+		return 0.f;
+	}
+
+	const uint64 Downloaded = DownloadedBytes();
+	return float(Downloaded) / float(TotalBytes);
+}
+
+float FInstallManagerBundleReport::GetInstallationProgress() const
+{
+	float TotalProgress = 0;
+	float TotalWeight = 0;
+	for (const FInstallBundleSourceReport& SourceReport : SourceReports)
+	{
+		for (const FInstallBundleReport& BundleReport : SourceReport.Bundles)
+		{
+			const float BundleWeight = BundleReport.TotalDownloadSize() / (1024.f * 1024.f * 1024.f);
+			const float BundleProgress = BundleReport.InstallationProgress();
+			TotalWeight += BundleWeight;
+			TotalProgress += BundleProgress * BundleWeight;
+		}
+	}
+	if (!TotalWeight)
+	{
+		return 0.f;
+	}
+	return TotalProgress / TotalWeight;
+}
+
+uint64 FInstallManagerBundleReport::TotalDownloadSize() const
+{
+	uint64 Total = 0;
+	for (const FInstallBundleSourceReport& SourceReport : SourceReports)
+	{
+		Total += SourceReport.TotalDownloadSize();
+	}
+	return Total;
+}
+
+uint64 FInstallManagerBundleReport::DownloadedBytes() const
+{
+	uint64 Total = 0;
+	for (const FInstallBundleSourceReport& SourceReport : SourceReports)
+	{
+		Total += SourceReport.DownloadedBytes();
+	}
+	return Total;
+}
+
+int FInstallManagerBundleReport::TotalBundleCount() const
+{
+	int Num = 0;
+	for (const FInstallBundleSourceReport& SourceReport : SourceReports)
+	{
+		Num += SourceReport.TotalBundleCount();
+	}
+	return Num;
+}
+
+int FInstallManagerBundleReport::UpdatedBundleCount() const
+{
+	int Num = 0;
+	for (const FInstallBundleSourceReport& SourceReport : SourceReports)
+	{
+		Num += SourceReport.UpdatedBundleCount();
+	}
+	return Num;
+}
+
+FString FInstallManagerBundleReport::GetProgressString() const
+{
+	const uint64 BytesDownloaded = DownloadedBytes();
+	const uint64 BytesTotal = TotalDownloadSize();
+	const int BundlesUpdated = UpdatedBundleCount();
+	const int BundlesTotal = TotalBundleCount();
+	const double Progress = BytesTotal ? double(BytesDownloaded) / BytesTotal : 0.0;
+	return FString::Printf(TEXT("{Sess.%d} INST %.2f%% DL %.2f%% [%.2f / %.2fGB] [%d / %d bundles]"),
+		LoadedCount,
+		100.f * GetInstallationProgress(),
+		100.f * Progress,
+		BytesDownloaded / (1000.f * 1000.f * 1000.f),
+		BytesTotal / (1000.f * 1000.f * 1000.f),
+		BundlesUpdated,
+		BundlesTotal);
+}
+
+FString FInstallManagerBundleReport::GetLongProgressString() const
+{
+	TStringBuilder<512> ProgressString;
+	ProgressString.Append(GetProgressString());
+	for (int i = 0; i < SourceReports.Num(); ++i)
+	{
+		if (SourceReports[i].Bundles.IsEmpty())
+		{
+			continue;
+		}
+
+		ProgressString.Appendf(TEXT(" [%d"), i);
+		for (int j = 0; j < SourceReports[i].Bundles.Num(); ++j)
+		{
+			ProgressString.Appendf(TEXT(" #%d,%s,%" UINT64_FMT "/%" UINT64_FMT "M"),
+				SourceReports[i].Bundles[j].State->GetIndex(),
+				*SourceReports[i].Bundles[j].BundleName.ToString().Right(7),
+				SourceReports[i].Bundles[j].DownloadedBytes() / (1000 * 1000),
+				SourceReports[i].Bundles[j].TotalDownloadSize() / (1000 * 1000));
+		}
+		ProgressString.Appendf(TEXT("]"));
+	}
+	return ProgressString.ToString();
+}
+
+FInstallManagerBundleReport::ECombinedStatus FInstallManagerBundleReport::GetCombinedStatus() const
+{
+	// If no sources or bundles, unknown
+	const int BundleCount = TotalBundleCount();
+	if (BundleCount == 0)
+	{
+		return ECombinedStatus::Unknown;
+	}
+
+	// Count numbers, if we find at least one unknown, we always return unknown, otherwise we look at the counters afterwards
+	int NumUpdatable = 0;
+	int NumUpdating = 0;
+	int NumInstalling = 0;
+	int NumUpdated = 0;
+	for (const FInstallBundleSourceReport& SourceReport : SourceReports)
+	{
+		for (const FInstallBundleReport& BundleReport : SourceReport.Bundles)
+		{
+			if (BundleReport.IsUnknown())
+			{
+				return ECombinedStatus::Unknown;
+			}
+			if (BundleReport.IsUpdatable())
+			{
+				++NumUpdatable;
+			}
+			else if (BundleReport.IsUpdating())
+			{
+				++NumUpdating;
+			}
+			else if (BundleReport.IsInstalling())
+			{
+				++NumInstalling;
+			}
+			else
+			{
+				++NumUpdated;
+			}
+		}
+	}
+
+	// Sanity check
+	check(BundleCount == NumUpdatable + NumUpdating + NumInstalling + NumUpdated);
+
+	// Determine state
+	if (NumUpdatable > 0 && NumUpdating == 0 && NumInstalling == 0 && NumUpdated == 0)
+	{
+		return ECombinedStatus::Updatable;
+	}
+	if (NumInstalling > 0 && NumUpdating == 0 && NumUpdatable == 0)
+	{
+		return ECombinedStatus::Finishing;
+	}
+	if (NumUpdated == BundleCount)
+	{
+		return ECombinedStatus::Finished;
+	}
+	return ECombinedStatus::Updating;
+}
+
+FInstallManagerBundleReportCache::FInstallManagerBundleReportCache(TArray<TUniquePtr<InstallBundleUtil::FInstallBundleTask>>& AsyncTasks)
+	: AsyncTasks(AsyncTasks)
+	, ReportFilename(FPaths::Combine(FPlatformMisc::GamePersistentDownloadDir(), TEXT("InstallBundleManagerReportCache.json"))) {
+}
+
+void FInstallManagerBundleReportCache::Load(FInstallManagerBundleReportCacheDoneLoadingDelegate OnLoadCompleteCallback)
+{
+	TSharedPtr<FInstallManagerBundleReport> NewReport = MakeShared<FInstallManagerBundleReport>();
+	InstallBundleUtil::StartInstallBundleAsyncIOTask(AsyncTasks,
+		[this, NewReport]()
+		{
+			if (FPaths::FileExists(ReportFilename))
+			{
+				FString JSONStringOnDisk;
+				if (FFileHelper::LoadFileToString(JSONStringOnDisk, *ReportFilename))
+				{
+					UE_LOG(LogInstallBundleManager, Display, TEXT("Reading InstallManagerBundleReportCache: %s"), *ReportFilename);
+					const bool bSuccess = NewReport->FromJson(JSONStringOnDisk);
+					ensureAlwaysMsgf(bSuccess, TEXT("Invalid JSON found while parsing InstallManagerBundleReportCache: %s"), *ReportFilename);
+				}
+				else
+				{
+					UE_LOG(LogInstallBundleManager, Warning, TEXT("Failed to read: %s"), *ReportFilename);
+				}
+			}
+		},
+		[this, NewReport, OnLoadCompleteCallback = MoveTemp(OnLoadCompleteCallback)]()
+		{
+			Report = MoveTemp(*NewReport);
+			OnLoadCompleteCallback.ExecuteIfBound();
+		}
+	);
+}
+
+void FInstallManagerBundleReportCache::SetReport(FInstallManagerBundleReport NewReport)
+{
+	Report = MoveTemp(NewReport);
+
+	InstallBundleUtil::StartInstallBundleAsyncIOTask(AsyncTasks,
+		[this, TargetReport = Report]()
+		{
+			UE_LOG(LogInstallBundleManager, Display, TEXT("Writing InstallManagerBundleReportCache: %s"), *ReportFilename);
+			const bool bSuccess = FFileHelper::SaveStringToFile(TargetReport.ToJson(), *ReportFilename);
+			return ensureAlwaysMsgf(bSuccess, TEXT("Error saving JSON output for InstallManagerBundleReportCache: %s"), *ReportFilename);
+		},
+		[] {});
+}
diff --git a/Engine/Source/Runtime/InstallBundleManager/Public/InstallBundleManagerInterface.h b/Engine/Source/Runtime/InstallBundleManager/Public/InstallBundleManagerInterface.h
--- a/Engine/Source/Runtime/InstallBundleManager/Public/InstallBundleManagerInterface.h
+++ b/Engine/Source/Runtime/InstallBundleManager/Public/InstallBundleManagerInterface.h
@@ -20,6 +20,7 @@
 
 class IAnalyticsProviderET;
 class IInstallBundleSource;
+struct FInstallManagerBundleReport;
 
 namespace InstallBundleUtil
 {
@@ -109,6 +110,8 @@ DECLARE_MULTICAST_DELEGATE_OneParam(FInstallBundleReleasedMultiDelegate, FInstal
 
 DECLARE_MULTICAST_DELEGATE_OneParam(FInstallBundleManagerOnPatchCheckComplete, EInstallBundleManagerPatchCheckResult);
 
+DECLARE_MULTICAST_DELEGATE_OneParam(FInstallBundleManagerReportingDelegate, const FInstallManagerBundleReport&);
+
 DECLARE_DELEGATE_RetVal(bool, FInstallBundleManagerEnvironmentWantsPatchCheck);
 
 DECLARE_DELEGATE_OneParam(FInstallBundleGetInstallStateDelegate, FInstallBundleCombinedInstallState);
@@ -125,6 +128,7 @@ public:
 	static INSTALLBUNDLEMANAGER_API FInstallBundlePausedMultiDelegate PausedBundleDelegate;
 	static INSTALLBUNDLEMANAGER_API FInstallBundleReleasedMultiDelegate ReleasedDelegate; // Called when content release request is complete
 	static INSTALLBUNDLEMANAGER_API FInstallBundleManagerOnPatchCheckComplete PatchCheckCompleteDelegate;
+	static INSTALLBUNDLEMANAGER_API FInstallBundleManagerReportingDelegate ReportingDelegate;
 
 	static INSTALLBUNDLEMANAGER_API TSharedPtr<IInstallBundleManager> GetPlatformInstallBundleManager();
 
@@ -134,6 +138,9 @@ public:
 
 	virtual bool HasBundleSource(FInstallBundleSourceType SourceType) const = 0;
 
+	virtual INSTALLBUNDLEMANAGER_API void SetContext(FName ContextName);
+	virtual INSTALLBUNDLEMANAGER_API bool CopyReportTo(FInstallManagerBundleReport& Report);
+
 	INSTALLBUNDLEMANAGER_API virtual const TSharedPtr<IInstallBundleSource> GetBundleSource(FInstallBundleSourceType SourceType) const;
 
 
diff --git a/Engine/Source/Runtime/InstallBundleManager/Public/InstallBundleManagerReporting.h b/Engine/Source/Runtime/InstallBundleManager/Public/InstallBundleManagerReporting.h
new file mode 100644
--- /dev/null
+++ b/Engine/Source/Runtime/InstallBundleManager/Public/InstallBundleManagerReporting.h
@@ -0,0 +1,248 @@
+// Copyright Epic Games, Inc. All Rights Reserved.
+
+#pragma once
+
+#include "InstallBundleUtils.h"
+
+class IInstallBundleSource;
+
+struct FInstallBundleReport : public FJsonSerializable
+{
+	// State that represents an updatable bundle, where a choice can be made to update in the background,
+	// or without.
+	struct FStateUpdatable : public FJsonSerializable
+	{
+		// The total amount of bytes to download for the bundle
+		uint64 DownloadSize;
+
+		BEGIN_JSON_SERIALIZER
+			JSON_SERIALIZE("DownloadSize", reinterpret_cast<int64&>(DownloadSize));
+		END_JSON_SERIALIZER
+	};
+
+	// State that represents a bundle that is updating
+	struct FStateUpdating : public FJsonSerializable
+	{
+		// The total amount of bytes that needs to be downloaded for the bundle
+		uint64 DownloadSize;
+
+		// Indicates if this bundle is being updated with background downloads or not
+		bool bUsesBackgroundDownloads = false;
+
+		// Indicates, in a value 0 <= x <= 1 the progress of the downloads. If this value reaches 1, the bundle
+		// can be assumed to be installed.
+		float DownloadProgress = 0;
+
+		// Indicates how many bytes were downloaded in the background
+		uint64 BackgroundDownloadedBytes;
+
+		BEGIN_JSON_SERIALIZER
+			JSON_SERIALIZE("DownloadSize", reinterpret_cast<int64&>(DownloadSize));
+			JSON_SERIALIZE("bUsesBackgroundDownloads", bUsesBackgroundDownloads);
+			JSON_SERIALIZE("DownloadProgress", DownloadProgress);
+			JSON_SERIALIZE("BackgroundDownloadedBytes", reinterpret_cast<int64&>(BackgroundDownloadedBytes));
+		END_JSON_SERIALIZER
+	};
+
+	struct FStateInstalling : public FJsonSerializable
+	{
+		// The total amount of bytes that have been downloaded for the bundle
+		uint64 DownloadedSize;
+
+		// The progress of the installation, if available, otherwise probably always 0; once the bundle switches to
+		// updated state, the InstallProgress is essentially 1
+		float InstallProgress = 0;
+
+		BEGIN_JSON_SERIALIZER
+			JSON_SERIALIZE("DownloadedSize", reinterpret_cast<int64&>(DownloadedSize));
+			JSON_SERIALIZE("InstallProgress", InstallProgress);
+		END_JSON_SERIALIZER
+	};
+
+	// State that represents a fully updated bundle that no longer needs to be updated
+	struct FStateUpdated : public FJsonSerializable
+	{
+		// The total amount of bytes that have been downloaded for the bundle
+		uint64 DownloadedSize;
+
+		BEGIN_JSON_SERIALIZER
+			JSON_SERIALIZE("DownloadedSize", reinterpret_cast<int64&>(DownloadedSize));
+		END_JSON_SERIALIZER
+	};
+
+	using FState = TVariant<FStateUpdatable, FStateUpdating, FStateInstalling, FStateUpdated>;
+
+	// The name of the bundle
+	FName BundleName;
+
+	// A string specifying what content version is being upgraded FROM. The format of this string is specific for the source.
+	// NOTE: if there is no known source version, this value should not be set.
+	TOptional<FString> SourceVersion;
+
+	// The state of the bundle
+	TOptional<FState> State;
+
+	BEGIN_JSON_SERIALIZER
+		JSON_SERIALIZE_NAME("Name", BundleName);
+		JSON_SERIALIZE_OPTIONAL("SourceVersion", SourceVersion);
+		JSON_SERIALIZE_OPTIONAL_VARIANT_BEGIN("State", State);
+			JSON_SERIALIZE_VARIANT_IFTYPE_SERIALIZABLE("StateUpdatable", FStateUpdatable);
+			JSON_SERIALIZE_VARIANT_IFTYPE_SERIALIZABLE("StateUpdating", FStateUpdating);
+			JSON_SERIALIZE_VARIANT_IFTYPE_SERIALIZABLE("StateInstalling", FStateInstalling);
+			JSON_SERIALIZE_VARIANT_IFTYPE_SERIALIZABLE("StateUpdated", FStateUpdated);
+		JSON_SERIALIZE_OPTIONAL_VARIANT_END();
+	END_JSON_SERIALIZER
+
+	// Returns the total amount of bytes that will or has been downloaded for this bundle.
+	INSTALLBUNDLEMANAGER_API uint64 TotalDownloadSize() const;
+
+	// Returns the amount of downloaded bytes for this bundle, should always be 0 <= x <= TotalDownloadSize()
+	INSTALLBUNDLEMANAGER_API uint64 DownloadedBytes() const;
+
+	// Returns a value 0 <= x <= 1 to indicate the installation progress.
+	// If the source doesn't support gradual installation progress, this will pop from 0 to 1.
+	INSTALLBUNDLEMANAGER_API float InstallationProgress() const;
+	
+	// Adds missing bytes to content size and download progress
+	INSTALLBUNDLEMANAGER_API void AddDiscrepancy(uint64 Bytes);
+	
+	bool IsUnknown() const { return !State.IsSet(); }
+	bool IsUpdatable() const { return State.IsSet() && State->IsType<FStateUpdatable>(); }
+	bool IsUpdating() const { return State.IsSet() && State->IsType<FStateUpdating>(); }
+	bool IsInstalling() const { return State.IsSet() && State->IsType<FStateInstalling>(); }
+	bool IsUpdated() const { return State.IsSet() && State->IsType<FStateUpdated>(); }
+};
+
+struct FInstallBundleSourceReport : public FJsonSerializable
+{
+	// The bundle source's type string
+	FName SourceType;
+
+	// A string specifying what content version is being upgrade TO. The format of this string is specific for the source.
+	// NOTE: if there is no known target version, this value should not be set.
+	TOptional<FString> TargetVersion;
+
+	// Array containing a single bundle report for every bundle this source is handling.
+	TArray<FInstallBundleReport> Bundles;
+
+	BEGIN_JSON_SERIALIZER
+		JSON_SERIALIZE_NAME("SourceType", SourceType);
+		JSON_SERIALIZE_OPTIONAL("TargetVersion", TargetVersion);
+		JSON_SERIALIZE_ARRAY_SERIALIZABLE("Bundles", Bundles, FInstallBundleReport);
+	END_JSON_SERIALIZER
+
+	// Returns pointer to bundle report that's valid as long as Bundles is not modified.
+	// Returns nullptr if there is no report yet for the given type.
+	INSTALLBUNDLEMANAGER_API FInstallBundleReport* TryFindBundleReport(FName BundleName);
+
+	// Returns a reference to the bundle source report for the given type. It adds one if it doesn't exist yet.
+	// The reference is valid as long as SourceReports is not modified.
+	INSTALLBUNDLEMANAGER_API FInstallBundleReport& FindOrAddBundleReport(FName BundleName);
+
+	// Returns the total amount of bytes that will or has been downloaded for this bundle source for all bundles
+	INSTALLBUNDLEMANAGER_API uint64 TotalDownloadSize() const;
+
+	// Returns the amount of downloaded bytes for this bundle source for all bundles, should always be 0 <= x <= TotalDownloadSize()
+	INSTALLBUNDLEMANAGER_API uint64 DownloadedBytes() const;
+
+	// Returns the total amount of bundles being updated, regardless of their state
+	INSTALLBUNDLEMANAGER_API int TotalBundleCount() const;
+
+	// Returns the amount of bundles that has being updated
+	INSTALLBUNDLEMANAGER_API int UpdatedBundleCount() const;
+};
+
+struct FInstallManagerBundleReport : public FJsonSerializable
+{
+	enum class ECombinedStatus
+	{
+		Unknown,		// if there are no bundle sources, or no bundles, or if at least one bundle is set to unknown
+		Updatable,		// if none of the above, if all bundles are set to updatable
+		Updating,		// if none of the above, if at least one bundle is still updating
+		Finishing,		// if none of the above, if at least one bundle is still installing
+		Finished,		// if none of the above, all bundles have been updated
+	};
+
+	// the context name for this bundle report
+	FName ContextName;
+
+	// Semi-unique session string, changes every time any bundle source decides to start updating to a new version; generated using UUID
+	FString SessionId;
+
+	// Indicates, if set, how many times the application was backgrounded during the update process for this bundle.
+	// If not set, the bundle source does not keep track of this.
+	TOptional<uint32> BackgroundedCount;
+
+	// Indicates how many times this report has been loaded from disk; in other words, indicates how many times the update process was restarted.
+	uint32 LoadedCount = 0;
+
+	// Array containing a single source report for every unique source bundle the manager knows.
+	TArray<FInstallBundleSourceReport> SourceReports;
+
+	BEGIN_JSON_SERIALIZER
+		JSON_SERIALIZE_NAME("ContextName", ContextName);
+		JSON_SERIALIZE("SessionId", SessionId);
+		JSON_SERIALIZE_OPTIONAL("BackgroundedCount", BackgroundedCount);
+		JSON_SERIALIZE("LoadedCount", LoadedCount);
+		JSON_SERIALIZE_ARRAY_SERIALIZABLE("SourceReports", SourceReports, FInstallBundleSourceReport);
+	END_JSON_SERIALIZER
+
+	// Returns pointer to bundle source report that's valid as long as SourceReports is not modified.
+	// Returns nullptr if there is no report yet for the given type.
+	INSTALLBUNDLEMANAGER_API FInstallBundleSourceReport* TryFindBundleSourceReport(const FInstallBundleSourceType& Type);
+
+	// Returns a reference to the bundle source report for the given type. It adds one if it doesn't exist yet.
+	// The reference is valid as long as SourceReports is not modified.
+	INSTALLBUNDLEMANAGER_API FInstallBundleSourceReport& FindOrAddBundleSourceReport(const FInstallBundleSourceType& Type);
+
+	// Returns the download progress.
+	// This returns a value 0 <= x <= 1.
+	INSTALLBUNDLEMANAGER_API float GetDownloadProgress() const;
+
+	// Returns the overall installation progress, which is separate from the download progress
+	// This returns a value 0 <= x <= 1.
+	INSTALLBUNDLEMANAGER_API float GetInstallationProgress() const;
+
+	// Returns the total amount of bytes that will or has been downloaded.
+	INSTALLBUNDLEMANAGER_API uint64 TotalDownloadSize() const;
+
+	// Returns the amount of downloaded bytes, should always be 0 <= x <= TotalDownloadSize()
+	INSTALLBUNDLEMANAGER_API uint64 DownloadedBytes() const;
+
+	// Returns the total amount of bundles being updated, regardless of their state
+	INSTALLBUNDLEMANAGER_API int TotalBundleCount() const;
+
+	// Returns the amount of bundles that has being updated
+	INSTALLBUNDLEMANAGER_API int UpdatedBundleCount() const;
+
+	// Returns a progress string
+	INSTALLBUNDLEMANAGER_API FString GetProgressString() const;
+
+	// Returns a longer progress string with more information
+	INSTALLBUNDLEMANAGER_API FString GetLongProgressString() const;
+
+	// Returns a combined status of all bundles
+	INSTALLBUNDLEMANAGER_API ECombinedStatus GetCombinedStatus() const;
+};
+
+DECLARE_DELEGATE(FInstallManagerBundleReportCacheDoneLoadingDelegate);
+
+class FInstallManagerBundleReportCache
+{
+public:
+	INSTALLBUNDLEMANAGER_API FInstallManagerBundleReportCache(TArray<TUniquePtr<InstallBundleUtil::FInstallBundleTask>>& AsyncTasks);
+
+	// Attempts to load the cached bundle report from disk; calls the delegate when it's done loading.
+	INSTALLBUNDLEMANAGER_API void Load(FInstallManagerBundleReportCacheDoneLoadingDelegate OnLoadCompleteCallback);
+
+	// Returns the current report as cached on disk (or last Set)
+	INSTALLBUNDLEMANAGER_API const FInstallManagerBundleReport& GetReport() const { return Report; }
+
+	// Sets the new report, both in memory, as well as flushing it to disk on a separate thread automatically.
+	INSTALLBUNDLEMANAGER_API void SetReport(FInstallManagerBundleReport NewReport);
+
+private:
+	FInstallManagerBundleReport Report;
+	TArray<TUniquePtr<InstallBundleUtil::FInstallBundleTask>>& AsyncTasks;
+	FString ReportFilename;
+};
diff --git a/Engine/Source/Runtime/InstallBundleManager/Public/InstallBundleSourceInterface.h b/Engine/Source/Runtime/InstallBundleManager/Public/InstallBundleSourceInterface.h
--- a/Engine/Source/Runtime/InstallBundleManager/Public/InstallBundleSourceInterface.h
+++ b/Engine/Source/Runtime/InstallBundleManager/Public/InstallBundleSourceInterface.h
@@ -7,6 +7,7 @@
 
 class IInstallBundleSource;
 class IAnalyticsProviderET;
+struct FInstallBundleReport;
 
 DECLARE_DELEGATE_TwoParams(FInstallBundleSourceInitDelegate, TSharedRef<IInstallBundleSource> /*Source*/, FInstallBundleSourceAsyncInitInfo /*InitInfo*/);
 
@@ -18,6 +19,8 @@ DECLARE_DELEGATE_TwoParams(FInstallBundleCompleteDelegate, TSharedRef<IInstallBu
 DECLARE_DELEGATE_TwoParams(FInstallBundlePausedDelegate, TSharedRef<IInstallBundleSource> /*Source*/, FInstallBundleSourcePauseInfo /*PauseInfo*/);
 DECLARE_DELEGATE_TwoParams(FInstallBundleRemovedDelegate, TSharedRef<IInstallBundleSource> /*Source*/, FInstallBundleSourceReleaseContentResultInfo /*Result*/);
 
+DECLARE_DELEGATE_OneParam(FInstallBundleReportingDelegate, const FInstallBundleReport&);
+
 DECLARE_DELEGATE_TwoParams(FInstallBundleSourceContentPatchResultDelegate, TSharedRef<IInstallBundleSource> /*Source*/, bool /*bContentPatchRequired*/);
 
 class IInstallBundleSource : public TSharedFromThis<IInstallBundleSource>
@@ -148,6 +151,9 @@ public:
 	virtual void GetDebugText(TArray<FString>& Output) {}
 #endif
 
+	virtual bool SupportsReportingDelegate() const { return false; }
+	virtual void SetReportingDelegate(FInstallBundleReportingDelegate Delegate) {}
+
 private:
 	TArray<FBuildInstallerStat> EmptyStats;
 };
diff --git a/Engine/Source/Runtime/InstallBundleManager/Public/InstallBundleTypes.h b/Engine/Source/Runtime/InstallBundleManager/Public/InstallBundleTypes.h
--- a/Engine/Source/Runtime/InstallBundleManager/Public/InstallBundleTypes.h
+++ b/Engine/Source/Runtime/InstallBundleManager/Public/InstallBundleTypes.h
@@ -195,6 +195,7 @@ enum class EInstallBundleGetContentStateFlags : uint32
 {
 	None = 0,
 	ForceNoPatching UE_DEPRECATED(5.7, "ForceNoPatching is no longer supported, use uninstall") = (1 << 0),
+	UpdateBundleReports = (1 << 1), // will update the BundleReport for the involved bundles (see InstallBundleManagerReporting.h and IInstallBundleManager::ReportingDelegate)
 };
 ENUM_CLASS_FLAGS(EInstallBundleGetContentStateFlags);
 
@@ -244,6 +245,7 @@ enum class EInstallBundleRequestFlags : uint32
 	TrackPersistentBundleStats = (1 << 4),
 	SkipMount = (1 << 5),
 	AsyncMount = (1 << 6),
+	UpdateBundleReports = (1 << 7), // will update the BundleReport for the involved bundles (see InstallBundleManagerReporting.h and IInstallBundleManager::ReportingDelegate)
 	Defaults = UseBackgroundDownloads,
 };
 ENUM_CLASS_FLAGS(EInstallBundleRequestFlags)
diff --git a/Engine/Source/Runtime/Json/Public/Serialization/JsonSerializerMacros.h b/Engine/Source/Runtime/Json/Public/Serialization/JsonSerializerMacros.h
--- a/Engine/Source/Runtime/Json/Public/Serialization/JsonSerializerMacros.h
+++ b/Engine/Source/Runtime/Json/Public/Serialization/JsonSerializerMacros.h
@@ -30,6 +30,19 @@
 #define JSON_SERIALIZE(JsonName, JsonValue) \
 		Serializer.Serialize(TEXTVIEW(JsonName), JsonValue)
 
+#define JSON_SERIALIZE_NAME(JsonName, JsonValue) \
+		if (Serializer.IsLoading()) \
+		{ \
+			FString NameStr; \
+			Serializer.Serialize(TEXTVIEW(JsonName), NameStr); \
+			JsonValue = FName(*NameStr); \
+		} \
+		else \
+		{ \
+			FString NameStr = JsonValue.ToString(); \
+			Serializer.Serialize(TEXTVIEW(JsonName), NameStr); \
+		}
+
 #define JSON_SERIALIZE_WITHDEFAULT(JsonName, JsonValue, DefaultJsonValue) \
 		if (Serializer.IsLoading()) \
 		{ \
@@ -410,3 +423,107 @@ else \
 	} \
 	Serializer.EndArray(); \
 }
+
+/**
+ * Private macros for Variant JSON serialization
+ */
+#define JSON_SERIALIZE_PRIVATE_VARIANT_BEGIN(JsonName, JsonValue) \
+	TSharedPtr<FJsonObject> VariantJsonObj; \
+	FString VariantTypeStr; \
+	if (VariantValue != nullptr) \
+	{ \
+		if (Serializer.IsLoading()) \
+		{ \
+			if (Serializer.GetObject()->HasTypedField<EJson::Object>(TEXTVIEW(JsonName))) \
+			{ \
+				VariantJsonObj = Serializer.GetObject()->GetObjectField(TEXTVIEW(JsonName)); \
+				if (VariantJsonObj.IsValid()) \
+				{ \
+					VariantTypeStr = VariantJsonObj->GetStringField(TEXTVIEW("Type")); \
+					const TSharedPtr<FJsonObject>* VariantValueObject = nullptr; \
+					if (VariantJsonObj->TryGetObjectField(TEXTVIEW("Value"), VariantValueObject); VariantValueObject != nullptr) \
+					{ \
+						VariantJsonObj = *VariantValueObject; \
+					} \
+					else \
+					{ \
+						VariantJsonObj = {}; \
+					} \
+				} \
+			} \
+		} \
+		else \
+		{ \
+			Serializer.StartObject(TEXTVIEW(JsonName)); \
+		} \
+	}
+
+#define JSON_SERIALIZE_PRIVATE_VARIANT_END() \
+	if (VariantValue != nullptr) \
+	{ \
+		if (!Serializer.IsLoading()) \
+		{ \
+			Serializer.EndObject(); \
+		} \
+	}
+
+/**
+ * Macros for begin/end of regular TVariant<...> value, insert JSON_SERIALIZE_VARIANT_IFTYPE_SERIALIZABLE in between BEGIN/END pair
+ */
+#define JSON_SERIALIZE_VARIANT_BEGIN(JsonName, JsonValue) \
+	if (true) \
+	{ \
+		auto* VariantValue = &JsonValue; \
+		JSON_SERIALIZE_PRIVATE_VARIANT_BEGIN(JsonName, JsonValue)
+
+#define JSON_SERIALIZE_VARIANT_END() \
+		JSON_SERIALIZE_PRIVATE_VARIANT_END() \
+	}
+
+/**
+ * Macros for begin/end of TOptional<TVariant<...>> value, insert JSON_SERIALIZE_VARIANT_IFTYPE_SERIALIZABLE in between BEGIN/END pair
+ */
+#define JSON_SERIALIZE_OPTIONAL_VARIANT_BEGIN(JsonName, JsonValue) \
+	if (true) \
+	{ \
+		if (Serializer.IsLoading()) \
+		{ \
+			if (Serializer.GetObject()->HasTypedField<EJson::Object>(TEXTVIEW(JsonName))) \
+			{ \
+				JsonValue.Emplace(); \
+			} \
+		} \
+		auto* VariantValue = JsonValue.IsSet() ? &JsonValue.GetValue() : nullptr; \
+		JSON_SERIALIZE_PRIVATE_VARIANT_BEGIN(JsonName, JsonValue)
+
+#define JSON_SERIALIZE_OPTIONAL_VARIANT_END() \
+		JSON_SERIALIZE_PRIVATE_VARIANT_END() \
+	}
+
+/**
+ * Macros for serializing a specific sub type of a variant, only to be insert in between (optional) VARIANT begin/end pair
+ */
+#define JSON_SERIALIZE_VARIANT_IFTYPE_SERIALIZABLE(TypeName, Type) \
+		if (VariantValue != nullptr) \
+		{ \
+			if (Serializer.IsLoading()) \
+			{ \
+				if (VariantJsonObj.IsValid() && VariantTypeStr == TEXTVIEW(TypeName)) \
+				{ \
+					VariantValue->Emplace<Type>(); \
+					VariantValue->Get<Type>().FromJson(VariantJsonObj); \
+				} \
+			} \
+			else \
+			{ \
+				/* created Value + Type */ \
+				if (VariantValue->IsType<Type>()) \
+				{ \
+					VariantTypeStr = TEXTVIEW(TypeName); \
+					Serializer.Serialize(TEXTVIEW("Type"), VariantTypeStr); \
+					Serializer.StartObject(TEXTVIEW("Value")); \
+					VariantValue->Get<Type>().Serialize(Serializer, true); \
+					Serializer.EndObject(); \
+				} \
+			} \
+		}
