Add CurrentSizeOnDiskOtherDirs so that bundles not installed under GamePersistentDownloadDir may report their install size.

#rnx

diff --git a/Engine/Source/Runtime/InstallBundleManager/Public/InstallBundleTypes.h b/Engine/Source/Runtime/InstallBundleManager/Public/InstallBundleTypes.h
--- a/Engine/Source/Runtime/InstallBundleManager/Public/InstallBundleTypes.h
+++ b/Engine/Source/Runtime/InstallBundleManager/Public/InstallBundleTypes.h
@@ -133,10 +133,12 @@ struct FInstallBundleContentSize
 	uint64 DownloadSize = 0;
 	// Size needed to install the remaining download under GamePersistentDownloadDir
 	uint64 SpaceRequiredForInstall = 0;
-	// Size needed to install the remaining download under other directories GamePersistentDownloadDir
+	// Size needed to install the remaining download under other directories than GamePersistentDownloadDir
 	uint64 SpaceRequiredForInstallOtherDirs = 0;
-	// size of bundle currently on disk.
+	// size of bundle currently on disk under GamePersistentDownloadDir
 	uint64 CurrentSizeOnDisk = 0;
+	// size of bundle currently on disk under other directories than GamePersistentDownloadDir
+	uint64 CurrentSizeOnDiskOtherDirs = 0;
 
 	FInstallBundleContentSize operator+(const FInstallBundleContentSize& Other) const
 	{
@@ -145,6 +147,7 @@ struct FInstallBundleContentSize
 		Result.SpaceRequiredForInstall = SpaceRequiredForInstall + Other.SpaceRequiredForInstall;
 		Result.SpaceRequiredForInstallOtherDirs = SpaceRequiredForInstallOtherDirs + Other.SpaceRequiredForInstallOtherDirs;
 		Result.CurrentSizeOnDisk = CurrentSizeOnDisk + Other.CurrentSizeOnDisk;
+		Result.CurrentSizeOnDiskOtherDirs = CurrentSizeOnDiskOtherDirs + Other.CurrentSizeOnDiskOtherDirs;
 		return Result;
 	}
 	FInstallBundleContentSize& operator+=(const FInstallBundleContentSize& Other)
