Installation progress stuckness fix

If user restarts the app after all bundles are downloaded, but not fully installed - combined progress code will recalculate weights incorrectly. Hence, combined progress will go over [0..1] range. The fix is to properly renormalizing weights. There might be more reasosn for this to happen.

Made stalls reportness less spammy and more controlable. Added bottom debug text to the log output, so we have more info.

#rnx

diff --git a/Engine/Source/Runtime/InstallBundleManager/Private/BundlePrereqCombinedStatusHelper.cpp b/Engine/Source/Runtime/InstallBundleManager/Private/BundlePrereqCombinedStatusHelper.cpp
--- a/Engine/Source/Runtime/InstallBundleManager/Private/BundlePrereqCombinedStatusHelper.cpp
+++ b/Engine/Source/Runtime/InstallBundleManager/Private/BundlePrereqCombinedStatusHelper.cpp
@@ -151,9 +151,10 @@ void FInstallBundleCombinedProgressTracker::SetBundlesToTrackFromContentState(co
 			RequiredBundleNames.Add(Bundle);
 
 			//If all bundles are zero weight, just treat this weight as 1 so everything ends up with 1 weight and is evenly distributed
-			CachedBundleWeights.FindOrAdd(Bundle) = bAreAllBundlesZeroWeight ? 1.0f : BundleState->Weight;
+			float Weight = bAreAllBundlesZeroWeight ? 1.0f : BundleState->Weight;
+			CachedBundleWeights.FindOrAdd(Bundle) = Weight;
 
-			TotalWeight += BundleState->Weight;
+			TotalWeight += Weight;
 		}
 	}
 
