Install Bundle Manager
- Add SkipOnDemandTocMount flag
- Fix corner case where SkipMount flag wasn't checked

diff --git a/Engine/Plugins/Experimental/DefaultInstallBundleManager/Source/Private/DefaultInstallBundleManager.cpp b/Engine/Plugins/Experimental/DefaultInstallBundleManager/Source/Private/DefaultInstallBundleManager.cpp
--- a/Engine/Plugins/Experimental/DefaultInstallBundleManager/Source/Private/DefaultInstallBundleManager.cpp
+++ b/Engine/Plugins/Experimental/DefaultInstallBundleManager/Source/Private/DefaultInstallBundleManager.cpp
@@ -1750,6 +1750,15 @@ void FDefaultInstallBundleManager::MountPaks(FContentRequestRef Request)
 
 	if (Request->OnDemandMountArgs.IsEmpty() || BundleInfo.bMountedOnDemandTocs || !AllowIoStoreOnDemandMount(Request, BundleInfo))
 	{
+		if (EnumHasAnyFlags(Request->Flags, EInstallBundleRequestFlags::SkipMount))
+		{
+			LOG_INSTALL_BUNDLE_MAN_OVERRIDE(Request->LogVerbosityOverride, Display, TEXT("Skipping Mounting Paks for Request %s"), *Request->BundleName.ToString());
+
+			TRACE_BOOKMARK(TEXT("Finished Mount Bundle %s"), *Request->BundleName.ToString());
+			Request->StepResult = EContentRequestStepResult::Done;
+			return;
+		}
+
 		LOG_INSTALL_BUNDLE_MAN_OVERRIDE(Request->LogVerbosityOverride, Display, TEXT("Mounting Paks for Request %s"), *BundleInfo.BundleNameString);
 		StatsBegin(Request->GetStatsKey(), EContentRequestState::Mounting);
 		StartMountPaksFunc();
@@ -1947,6 +1956,11 @@ void FDefaultInstallBundleManager::UnmountPaks(FContentReleaseRequestRef Request
 
 bool FDefaultInstallBundleManager::AllowIoStoreOnDemandMount(FContentRequestRef Request, const FBundleInfo& BundleInfo)
 {
+	if (EnumHasAnyFlags(Request->Flags, EInstallBundleRequestFlags::SkipOnDemandTocMount))
+	{
+		return false;
+	}
+
 	return true;
 }
 
diff --git a/Engine/Source/Runtime/InstallBundleManager/Public/InstallBundleTypes.h b/Engine/Source/Runtime/InstallBundleManager/Public/InstallBundleTypes.h
--- a/Engine/Source/Runtime/InstallBundleManager/Public/InstallBundleTypes.h
+++ b/Engine/Source/Runtime/InstallBundleManager/Public/InstallBundleTypes.h
@@ -242,17 +242,18 @@ INSTALLBUNDLEMANAGER_API const TCHAR* LexToString(EInstallBundleResult Result);
 // TODO: Should probably be renamed to EInstallBundleRequestUpdateFlags 
 enum class EInstallBundleRequestFlags : uint32
 {
-	None = 0,
-	CheckForCellularDataUsage = (1 << 0),
+	None											= 0,
+	CheckForCellularDataUsage						= (1 << 0),
 	UseBackgroundDownloads UE_DEPRECATED(5.7, "UseBackgroundDownloads is no longer supported, done automatically") = (1 << 1),
-	SendNotificationIfDownloadCompletesInBackground = (1 << 2),
+	SendNotificationIfDownloadCompletesInBackground	= (1 << 2),
 	ForceNoPatching UE_DEPRECATED(5.7, "ForceNoPatching is no longer supported, use uninstall") = (1 << 3),
-	TrackPersistentBundleStats = (1 << 4),
-	SkipMount = (1 << 5),
-	AsyncMount = (1 << 6),
-	UpdateBundleReports = (1 << 7), // will update the BundleReport for the involved bundles (see InstallBundleManagerReporting.h and IInstallBundleManager::ReportingDelegate)
-	ExplicitUpdateList = (1 << 8),
-	Defaults = None,
+	TrackPersistentBundleStats						= (1 << 4),
+	SkipMount										= (1 << 5),
+	SkipOnDemandTocMount							= (1 << 6),
+	AsyncMount										= (1 << 7),
+	UpdateBundleReports								= (1 << 8), // will update the BundleReport for the involved bundles (see InstallBundleManagerReporting.h and IInstallBundleManager::ReportingDelegate)
+	ExplicitUpdateList								= (1 << 9),
+	Defaults										= None,
 };
 ENUM_CLASS_FLAGS(EInstallBundleRequestFlags)
 
