Updated telemetry values for Lowest and Highest bundle version

#rnx
#Synced-CL 45098147

diff --git a/Engine/Source/Runtime/InstallBundleManager/Private/InstallBundleManagerReporting.cpp b/Engine/Source/Runtime/InstallBundleManager/Private/InstallBundleManagerReporting.cpp
--- a/Engine/Source/Runtime/InstallBundleManager/Private/InstallBundleManagerReporting.cpp
+++ b/Engine/Source/Runtime/InstallBundleManager/Private/InstallBundleManagerReporting.cpp
@@ -361,13 +361,12 @@ TOptional<FString> FInstallManagerBundleReport::GetLowestCurrentBundleVersion()
 	{
 		for (const FInstallBundleReport& BundleReport : SourceReport.Bundles)
 		{
-			if (!BundleReport.SourceVersion.IsSet() || BundleReport.SourceVersion.GetValue().IsEmpty())
+			if (BundleReport.SourceVersion.IsSet() && !BundleReport.SourceVersion.GetValue().IsEmpty())
 			{
-				return {};
-			}
-			if (!Lowest.IsSet() || BundleReport.SourceVersion.GetValue() < Lowest.GetValue())
-			{
-				Lowest = BundleReport.SourceVersion.GetValue();
+				if (!Lowest.IsSet() || BundleReport.SourceVersion.GetValue() < Lowest.GetValue())
+				{
+					Lowest = BundleReport.SourceVersion.GetValue();
+				}
 			}
 		}
 	}
@@ -381,7 +380,7 @@ TOptional<FString> FInstallManagerBundleReport::GetHighestCurrentBundleVersion()
 	{
 		for (const FInstallBundleReport& BundleReport : SourceReport.Bundles)
 		{
-			if (BundleReport.SourceVersion.IsSet())
+			if (BundleReport.SourceVersion.IsSet() && !BundleReport.SourceVersion.GetValue().IsEmpty())
 			{
 				if (!Highest.IsSet() || BundleReport.SourceVersion.GetValue() > Highest.GetValue())
 				{
