Added some comments

#Synced-CL 50389196

diff --git a/Engine/Source/Runtime/InstallBundleManager/Public/InstallBundleTypes.h b/Engine/Source/Runtime/InstallBundleManager/Public/InstallBundleTypes.h
--- a/Engine/Source/Runtime/InstallBundleManager/Public/InstallBundleTypes.h
+++ b/Engine/Source/Runtime/InstallBundleManager/Public/InstallBundleTypes.h
@@ -505,6 +505,8 @@ enum class UE_DEPRECATED(5.7, "Use EInstallBundleCacheStatsFlags") EInstallBundl
 	CSV
 };
 
+// TODO: These details shouldn't be in engine code. If you find yourself adding to this struct, please consider refactoring this 
+// to pass back a list of FAnalyticsEventAttribute's instead.
 struct FBundleDecompressionInfo
 {
 	uint64 TotalCompressedBytes{};				// total amount of bytes in compressed files (incl headers)
@@ -516,6 +518,8 @@ struct FBundleDecompressionInfo
 	double TotalDecompressionWriteTime{};		// seconds spent writing decompressed bytes this session
 };
 
+// TODO: These details shouldn't be in engine code. If you find yourself adding to this struct, please consider refactoring this 
+// to pass back a list of FAnalyticsEventAttribute's instead.
 struct FBundleRequestCompleteInfo
 {
 	FString PreviousManifest;
