Virtual bundle client-side support

#Synced-CL 50931132

#Synced-CL 50931132

#Synced-CL 50931132

diff --git a/Engine/Source/Runtime/InstallBundleManager/Public/InstallBundleUtils.h b/Engine/Source/Runtime/InstallBundleManager/Public/InstallBundleUtils.h
--- a/Engine/Source/Runtime/InstallBundleManager/Public/InstallBundleUtils.h
+++ b/Engine/Source/Runtime/InstallBundleManager/Public/InstallBundleUtils.h
@@ -710,6 +710,9 @@ namespace InstallBundleUtil
 		FContentRequestSharedContext& operator=(FContentRequestSharedContext&& Other) = default;
 
 		TMap<FInstallBundleSourceType, TUniquePtr<IBundleSourceContentRequestSharedContext>> BundleSourceSharedContext;
+
+		// (Optional) - can be used to pass a list of requested bundles to IInstallBundleSource implementations
+		TArray<FName> RequestedBundles;	
 	};
 	using FContentRequestSharedContextPtr = TSharedPtr<FContentRequestSharedContext>;
 } // InstallBundleUtil
